Fix build include issue on some platforms (Horde 519410)

#rnx

diff --git a/Engine/Source/Runtime/AudioDeviceEnumeration/Windows/WindowsMMDeviceEnumeration/Private/WindowsMMDeviceEnumerationModule.cpp b/Engine/Source/Runtime/AudioDeviceEnumeration/Windows/WindowsMMDeviceEnumeration/Private/WindowsMMDeviceEnumerationModule.cpp
--- a/Engine/Source/Runtime/AudioDeviceEnumeration/Windows/WindowsMMDeviceEnumeration/Private/WindowsMMDeviceEnumerationModule.cpp
+++ b/Engine/Source/Runtime/AudioDeviceEnumeration/Windows/WindowsMMDeviceEnumeration/Private/WindowsMMDeviceEnumerationModule.cpp
@@ -1,10 +1,8 @@
 // Copyright Epic Games, Inc. All Rights Reserved.
 
 #include "Modules/ModuleManager.h"
-#include "Windows/WindowsPlatformMisc.h"
 
 #if PLATFORM_WINDOWS
-
 #include "Windows/AllowWindowsPlatformTypes.h"
 
 THIRD_PARTY_INCLUDES_START
@@ -13,7 +11,6 @@ THIRD_PARTY_INCLUDES_START
 THIRD_PARTY_INCLUDES_END
 
 #include "Windows/HideWindowsPlatformTypes.h"
-
 #endif // PLATFORM_WINDOWS
 
 #include "WindowsMMNotificationClient.h"
@@ -26,22 +23,13 @@ namespace Audio
 {
 	class FWindowsMMDeviceEnumerationModule : public IModuleInterface
 	{
-	private:
-		/** Indicates if FWindowsPlatformMisc::CoInitialize() was successfull. */
-		bool bCoInitialized = false;
-
 	public:
 		virtual void StartupModule() override
 		{
-			bCoInitialized = FWindowsPlatformMisc::CoInitialize();
 		}
 
 		virtual void ShutdownModule() override
 		{
-			if (bCoInitialized)
-			{
-				FWindowsPlatformMisc::CoUninitialize();
-			}
 		}
 	};
 }
