Fix device enumeration logic to use the correct cvar.

#rnx

diff --git a/Engine/Source/Runtime/AudioDeviceEnumeration/Windows/WindowsMMDeviceEnumeration/Private/WindowsMMCvarUtils.cpp b/Engine/Source/Runtime/AudioDeviceEnumeration/Windows/WindowsMMDeviceEnumeration/Private/WindowsMMCvarUtils.cpp
--- a/Engine/Source/Runtime/AudioDeviceEnumeration/Windows/WindowsMMDeviceEnumeration/Private/WindowsMMCvarUtils.cpp
+++ b/Engine/Source/Runtime/AudioDeviceEnumeration/Windows/WindowsMMDeviceEnumeration/Private/WindowsMMCvarUtils.cpp
@@ -3,6 +3,13 @@
 #include "WindowsMMCvarUtils.h"
 #include "HAL/IConsoleManager.h"
 
+static int32 DisableDeviceSwapCVar = 0;
+FAutoConsoleVariableRef CVarDisableDeviceSwap(
+	TEXT("au.DisableDeviceSwap"),
+	DisableDeviceSwapCVar,
+	TEXT("Disable device swap handling code for Audio Mixer on Windows.\n")
+	TEXT("0: Not Enabled, 1: Enabled"),
+	ECVF_Default);
 
 static int32 EnableDetailedWindowsDeviceLoggingCVar = 0;
 FAutoConsoleVariableRef CVarEnableDetailedWindowsDeviceLogging(
@@ -14,6 +21,11 @@ FAutoConsoleVariableRef CVarEnableDetailedWindowsDeviceLogging(
 
 namespace Audio
 {
+	bool WindowsMMCvarUtils::ShouldIgnoreDeviceSwaps()
+	{
+		return DisableDeviceSwapCVar != 0;
+	}
+
 	bool WindowsMMCvarUtils::ShouldLogDeviceSwaps()
 	{
 		return EnableDetailedWindowsDeviceLoggingCVar != 0;
diff --git a/Engine/Source/Runtime/AudioDeviceEnumeration/Windows/WindowsMMDeviceEnumeration/Private/WindowsMMCvarUtils.h b/Engine/Source/Runtime/AudioDeviceEnumeration/Windows/WindowsMMDeviceEnumeration/Private/WindowsMMCvarUtils.h
--- a/Engine/Source/Runtime/AudioDeviceEnumeration/Windows/WindowsMMDeviceEnumeration/Private/WindowsMMCvarUtils.h
+++ b/Engine/Source/Runtime/AudioDeviceEnumeration/Windows/WindowsMMDeviceEnumeration/Private/WindowsMMCvarUtils.h
@@ -8,7 +8,12 @@ namespace Audio
 	class WindowsMMCvarUtils
 	{
 	public:
+		/** Called by FWindowsMMNotificationClient to bypass notifications for audio device changes: */
+		static bool ShouldIgnoreDeviceSwaps();
+
+		/** Called by FWindowsMMNotificationClient to toggle logging for audio device changes: */
 		static bool ShouldLogDeviceSwaps();
+
 	};
 
 } // namespace Audio
\ No newline at end of file
diff --git a/Engine/Source/Runtime/AudioDeviceEnumeration/Windows/WindowsMMDeviceEnumeration/Private/WindowsMMNotificationClient.cpp b/Engine/Source/Runtime/AudioDeviceEnumeration/Windows/WindowsMMDeviceEnumeration/Private/WindowsMMNotificationClient.cpp
--- a/Engine/Source/Runtime/AudioDeviceEnumeration/Windows/WindowsMMDeviceEnumeration/Private/WindowsMMNotificationClient.cpp
+++ b/Engine/Source/Runtime/AudioDeviceEnumeration/Windows/WindowsMMDeviceEnumeration/Private/WindowsMMNotificationClient.cpp
@@ -302,7 +302,7 @@ namespace Audio
 	{
 		UE_CLOG(WindowsMMCvarUtils::ShouldLogDeviceSwaps(), LogAudioEnumeration, Display, TEXT("FWindowsMMNotificationClient: OnDeviceStateChanged: %s, %d"), *GetFriendlyName(pwstrDeviceId), dwNewState);
 
-		if (WindowsMMCvarUtils::ShouldLogDeviceSwaps())
+		if (WindowsMMCvarUtils::ShouldIgnoreDeviceSwaps())
 		{
 			return S_OK;
 		}
@@ -326,7 +326,7 @@ namespace Audio
 	{
 		UE_CLOG(WindowsMMCvarUtils::ShouldLogDeviceSwaps(), LogAudioEnumeration, Display, TEXT("FWindowsMMNotificationClient: OnDeviceRemoved: %s"), *GetFriendlyName(pwstrDeviceId));
 
-		if (WindowsMMCvarUtils::ShouldLogDeviceSwaps())
+		if (WindowsMMCvarUtils::ShouldIgnoreDeviceSwaps())
 		{
 			return S_OK;
 		}
@@ -345,7 +345,7 @@ namespace Audio
 	{
 		UE_CLOG(WindowsMMCvarUtils::ShouldLogDeviceSwaps(), LogAudioEnumeration, Display, TEXT("FWindowsMMNotificationClient: OnDeviceAdded: %s"), *GetFriendlyName(pwstrDeviceId));
 
-		if (WindowsMMCvarUtils::ShouldLogDeviceSwaps())
+		if (WindowsMMCvarUtils::ShouldIgnoreDeviceSwaps())
 		{
 			return S_OK;
 		}
@@ -385,7 +385,7 @@ namespace Audio
 
 		Audio::EAudioDeviceRole AudioDeviceRole;
 
-		if (WindowsMMCvarUtils::ShouldLogDeviceSwaps())
+		if (WindowsMMCvarUtils::ShouldIgnoreDeviceSwaps())
 		{
 			return S_OK;
 		}
diff --git a/Engine/Source/Runtime/AudioMixerCore/Private/AudioMixer.cpp b/Engine/Source/Runtime/AudioMixerCore/Private/AudioMixer.cpp
--- a/Engine/Source/Runtime/AudioMixerCore/Private/AudioMixer.cpp
+++ b/Engine/Source/Runtime/AudioMixerCore/Private/AudioMixer.cpp
@@ -63,14 +63,6 @@ FAutoConsoleVariableRef CVarRenderThreadAffinity(
 	TEXT("0: Disabled (Default), otherwise overriden thread affinity."),
 	ECVF_Default);
 
-static int32 DisableDeviceSwapCVar = 0;
-FAutoConsoleVariableRef CVarDisableDeviceSwap(
-	TEXT("au.DisableDeviceSwap"),
-	DisableDeviceSwapCVar,
-	TEXT("Disable device swap handling code for Audio Mixer on Windows.\n")
-	TEXT("0: Not Enabled, 1: Enabled"),
-	ECVF_Default);
-
 static int32 bUseThreadedDeviceSwapCVar = 1;
 FAutoConsoleVariableRef CVarUseThreadedDeviceSwap(
 	TEXT("au.UseThreadedDeviceSwap"),
@@ -882,11 +874,6 @@ namespace Audio
 		return false;
 	}
 
-	bool IAudioMixer::ShouldIgnoreDeviceSwaps()
-	{
-		return DisableDeviceSwapCVar != 0;
-	}
-
 	bool IAudioMixer::ShouldUseThreadedDeviceSwap()
 	{
 		return bUseThreadedDeviceSwapCVar != 0;
diff --git a/Engine/Source/Runtime/AudioMixerCore/Public/AudioMixer.h b/Engine/Source/Runtime/AudioMixerCore/Public/AudioMixer.h
--- a/Engine/Source/Runtime/AudioMixerCore/Public/AudioMixer.h
+++ b/Engine/Source/Runtime/AudioMixerCore/Public/AudioMixer.h
@@ -203,12 +203,6 @@ namespace Audio
 
 		bool IsMainAudioMixer() const { return bIsMainAudioMixer; }
 
-		/** Called by FWindowsMMNotificationClient to bypass notifications for audio device changes: */
-		AUDIOMIXERCORE_API static bool ShouldIgnoreDeviceSwaps();
-
-		/** Called by FWindowsMMNotificationClient to toggle logging for audio device changes: */
-		AUDIOMIXERCORE_API static bool ShouldLogDeviceSwaps();
-		
 		/** Called by AudioMixer to see if we should do a multithreaded device swap */
 		AUDIOMIXERCORE_API static bool ShouldUseThreadedDeviceSwap();
 
