Only enumerate audio hardware when neccessary.

#rnx

diff --git a/Engine/Source/Runtime/AudioDeviceEnumeration/Windows/WindowsMMDeviceEnumeration/Private/WindowsMMDeviceInfoCache.cpp b/Engine/Source/Runtime/AudioDeviceEnumeration/Windows/WindowsMMDeviceEnumeration/Private/WindowsMMDeviceInfoCache.cpp
--- a/Engine/Source/Runtime/AudioDeviceEnumeration/Windows/WindowsMMDeviceEnumeration/Private/WindowsMMDeviceInfoCache.cpp
+++ b/Engine/Source/Runtime/AudioDeviceEnumeration/Windows/WindowsMMDeviceEnumeration/Private/WindowsMMDeviceInfoCache.cpp
@@ -71,7 +71,8 @@ FWindowsMMDeviceCache::FCacheEntry::FCacheEntry(const FCacheEntry& InOther)
 	*this = InOther;
 }
 
-FWindowsMMDeviceCache::FWindowsMMDeviceCache()
+FWindowsMMDeviceCache::FWindowsMMDeviceCache(bool bInEnumerateHardware)
+	: bEnumerateHardware(bInEnumerateHardware)
 {
 	ensure(SUCCEEDED(CoCreateInstance(__uuidof(MMDeviceEnumerator), NULL, CLSCTX_INPROC_SERVER, IID_PPV_ARGS(&DeviceEnumerator))) && DeviceEnumerator);
 
@@ -291,6 +292,8 @@ bool FWindowsMMDeviceCache::EnumerateDeviceProps(const TComPtr<IMMDevice>& InDev
 
 bool FWindowsMMDeviceCache::EnumerateHardwareTopology(const TComPtr<IMMDevice>& InDevice, FCacheEntry& OutInfo)
 {
+	SCOPED_NAMED_EVENT(FWindowsMMDeviceCache_EnumerateHardwareTopology, FColor::Blue);
+
 	if (!DeviceEnumerator)
 	{
 		return false;
@@ -358,6 +361,8 @@ bool FWindowsMMDeviceCache::EnumerateHardwareTopology(const TComPtr<IMMDevice>&
 
 void FWindowsMMDeviceCache::EnumerateEndpoints()
 {
+	SCOPED_NAMED_EVENT(FWindowsMMDeviceCache_EnumerateEndpoints, FColor::Blue);
+
 	// Build a new cache from scratch.
 	TMap<FName, FCacheEntry> NewCache;
 
@@ -385,7 +390,10 @@ void FWindowsMMDeviceCache::EnumerateEndpoints()
 						EnumerateDeviceProps(Device, Info);
 
 						// Enumerate hardware topology to fetch hardware Id.
-						EnumerateHardwareTopology(Device, Info);
+						if (bEnumerateHardware)
+						{
+							EnumerateHardwareTopology(Device, Info);
+						}
 
 						UE_LOG(LogAudioEnumeration, Verbose, TEXT("FWindowsMMDeviceCache: %s Device '%s' ID='%s'"),
 							Info.Type == FCacheEntry::EEndpointType::Capture ? TEXT("Capture") :
diff --git a/Engine/Source/Runtime/AudioDeviceEnumeration/Windows/WindowsMMDeviceEnumeration/Public/WindowsMMDeviceInfoCache.h b/Engine/Source/Runtime/AudioDeviceEnumeration/Windows/WindowsMMDeviceEnumeration/Public/WindowsMMDeviceInfoCache.h
--- a/Engine/Source/Runtime/AudioDeviceEnumeration/Windows/WindowsMMDeviceEnumeration/Public/WindowsMMDeviceInfoCache.h
+++ b/Engine/Source/Runtime/AudioDeviceEnumeration/Windows/WindowsMMDeviceEnumeration/Public/WindowsMMDeviceInfoCache.h
@@ -63,8 +63,9 @@ struct WINDOWSMMDEVICEENUMERATION_API FWindowsMMDeviceCache : IAudioMixerDeviceC
 	TMap<FName, FCacheEntry> Cache;								// DeviceID GUID -> Info.
 	FName DefaultCaptureId[(int32)EAudioDeviceRole::COUNT];		// Role -> DeviceID GUID
 	FName DefaultRenderId[(int32)EAudioDeviceRole::COUNT];		// Role -> DeviceID GUID
+	bool bEnumerateHardware = false;							// Hardware enumeration can be slow. Enumerate only if requested.
 
-	FWindowsMMDeviceCache();
+	explicit FWindowsMMDeviceCache(bool bInEnumerateHardware = false);
 	virtual ~FWindowsMMDeviceCache() = default;
 
 	bool EnumerateChannelMask(uint32 InMask, FCacheEntry& OutInfo);
diff --git a/Engine/Source/Runtime/Windows/AudioMixerWasapi/Private/AudioMixerWasapi.cpp b/Engine/Source/Runtime/Windows/AudioMixerWasapi/Private/AudioMixerWasapi.cpp
--- a/Engine/Source/Runtime/Windows/AudioMixerWasapi/Private/AudioMixerWasapi.cpp
+++ b/Engine/Source/Runtime/Windows/AudioMixerWasapi/Private/AudioMixerWasapi.cpp
@@ -142,6 +142,8 @@ namespace Audio
 
 	bool FAudioMixerWasapi::InitStreamParams(const uint32 InDeviceIndex, const int32 InNumBufferFrames, const int32 InNumBuffers, const int32 InSampleRate, TArray<FWasapiRenderStreamParams>& OutParams)
 	{
+		SCOPED_NAMED_EVENT(FAudioMixerWasapi_InitStreamParams, FColor::Blue);
+
 		FAudioPlatformDeviceInfo DeviceInfo;
 		if (!GetOutputDeviceInfo(InDeviceIndex, DeviceInfo))
 		{
@@ -186,6 +188,8 @@ namespace Audio
 
 	bool FAudioMixerWasapi::OpenAudioStream(const FAudioMixerOpenStreamParams& Params)
 	{
+		SCOPED_NAMED_EVENT(FAudioMixerWasapi_OpenAudioStream, FColor::Green);
+
 		OpenStreamParams = Params;
 
 		AudioStreamInfo.Reset();
diff --git a/Engine/Source/Runtime/Windows/AudioMixerWasapi/Private/AudioMixerWasapiDeviceListener.cpp b/Engine/Source/Runtime/Windows/AudioMixerWasapi/Private/AudioMixerWasapiDeviceListener.cpp
--- a/Engine/Source/Runtime/Windows/AudioMixerWasapi/Private/AudioMixerWasapiDeviceListener.cpp
+++ b/Engine/Source/Runtime/Windows/AudioMixerWasapi/Private/AudioMixerWasapiDeviceListener.cpp
@@ -40,7 +40,7 @@ namespace Audio
 		if (!DeviceInfoCache.IsValid())
 		{
 			// Setup device info cache.
-			DeviceInfoCache = MakeUnique<FWindowsMMDeviceCache>();
+			DeviceInfoCache = MakeUnique<FWindowsMMDeviceCache>(bUseWasapiAggregateDevice);
 			WasapiWinNotificationClient->RegisterDeviceChangedListener(static_cast<FWindowsMMDeviceCache*>(DeviceInfoCache.Get()));
 		}
 
