Adding ability to open audio renderer with 8 channels even if output is stereo

[RN] Added new command line argument "ForceSurroundSound", which adds the ability to force a windows machine to initialize spoof the audio device to render 8 channels of audio. Useful for scenarios where an 8-channel audio device isn't available and a developer needs to exercise 8-channel audio rendering code.

diff --git a/Engine/Source/Runtime/AudioDeviceEnumeration/Windows/WindowsMMDeviceEnumeration/Private/WindowsMMDeviceInfoCache.cpp b/Engine/Source/Runtime/AudioDeviceEnumeration/Windows/WindowsMMDeviceEnumeration/Private/WindowsMMDeviceInfoCache.cpp
--- a/Engine/Source/Runtime/AudioDeviceEnumeration/Windows/WindowsMMDeviceEnumeration/Private/WindowsMMDeviceInfoCache.cpp
+++ b/Engine/Source/Runtime/AudioDeviceEnumeration/Windows/WindowsMMDeviceEnumeration/Private/WindowsMMDeviceInfoCache.cpp
@@ -2,26 +2,37 @@
 
 #include "WindowsMMDeviceInfoCache.h"
 
-#include "Windows/AllowWindowsPlatformTypes.h"
-#include "Windows/AllowWindowsPlatformAtomics.h"
-
 THIRD_PARTY_INCLUDES_START
 #include <devicetopology.h>
 #include <functiondiscoverykeys_devpkey.h>
 THIRD_PARTY_INCLUDES_END
 
+#include "ConversionHelpers.h"
+#include "HAL/IConsoleManager.h"
+#include "Misc/CommandLine.h"
+#include "Misc/Parse.h"
+#include "Misc/ScopeRWLock.h"
+#include "Windows/AllowWindowsPlatformTypes.h"
+#include "Windows/AllowWindowsPlatformAtomics.h"
 #include "Windows/HideWindowsPlatformAtomics.h"
 #include "Windows/HideWindowsPlatformTypes.h"
-
 #include "WindowsMMStringUtils.h"
-#include "ConversionHelpers.h"
-#include "Misc/ScopeRWLock.h"
 #include "WindowsMMCvarUtils.h"
 #include "WindowsMMDeviceEnumerationLog.h"
 
+
 namespace Audio
 {
 
+// Support for forcing the audio renderer on windows machines to initialize as if it was connected to a 7.1 output device.
+// This is useful for cases where a developer does not have access to a 7.1 device (or is working on a cloud machine) and 
+// wants to validate and excersie surround sound audio rendering code.
+static bool GetForceSurroundSound()
+{	
+	static bool bForceSurroundSound = FParse::Param(FCommandLine::Get(), TEXT("ForceSurroundSound"));
+	return bForceSurroundSound;
+}
+
 FWindowsMMDeviceCache::FCacheEntry& FWindowsMMDeviceCache::FCacheEntry::operator=(FCacheEntry&& InOther)
 {
 	DeviceId = MoveTemp(InOther.DeviceId);
@@ -264,7 +275,14 @@ bool FWindowsMMDeviceCache::EnumerateDeviceProps(const TComPtr<IMMDevice>& InDev
 			if (SUCCEEDED(InPropStore->GetValue(InKey, &DeviceFormat)) && DeviceFormat.blob.pBlobData)
 			{
 				const WAVEFORMATEX* WaveFormatEx = (const WAVEFORMATEX*)(DeviceFormat.blob.pBlobData);
-				OutInfo.NumChannels = FMath::Clamp((int32)WaveFormatEx->nChannels, 2, 8);
+				if (GetForceSurroundSound())
+				{
+					OutInfo.NumChannels = 8;
+				}
+				else
+				{
+					OutInfo.NumChannels = FMath::Clamp((int32)WaveFormatEx->nChannels, 2, 8);
+				}
 				OutInfo.SampleRate = WaveFormatEx->nSamplesPerSec;
 
 				EnumerateChannelFormat(WaveFormatEx, OutInfo);
