[AudioDeviceEnumeration]
* Moved function to cpp file since it can't be both inlined and dllexported at the same time and inline caused other problems

diff --git a/Engine/Source/Runtime/AudioDeviceEnumeration/Windows/WindowsMMDeviceEnumeration/Private/WindowsMMNotificationClient.cpp b/Engine/Source/Runtime/AudioDeviceEnumeration/Windows/WindowsMMDeviceEnumeration/Private/WindowsMMNotificationClient.cpp
--- a/Engine/Source/Runtime/AudioDeviceEnumeration/Windows/WindowsMMDeviceEnumeration/Private/WindowsMMNotificationClient.cpp
+++ b/Engine/Source/Runtime/AudioDeviceEnumeration/Windows/WindowsMMDeviceEnumeration/Private/WindowsMMNotificationClient.cpp
@@ -247,6 +247,11 @@ namespace Audio
 		return {};
 	}
 
+	uint32 FWindowsMMNotificationClient::ReleaseClient()
+	{
+		return Release();
+	}
+
 	FString FWindowsMMNotificationClient::GetFriendlyName(const TComPtr<IMMDevice>& InDevice)
 	{
 		FString FriendlyName = TEXT("[No Friendly Name for Device]");
diff --git a/Engine/Source/Runtime/AudioDeviceEnumeration/Windows/WindowsMMDeviceEnumeration/Public/WindowsMMNotificationClient.h b/Engine/Source/Runtime/AudioDeviceEnumeration/Windows/WindowsMMDeviceEnumeration/Public/WindowsMMNotificationClient.h
--- a/Engine/Source/Runtime/AudioDeviceEnumeration/Windows/WindowsMMDeviceEnumeration/Public/WindowsMMNotificationClient.h
+++ b/Engine/Source/Runtime/AudioDeviceEnumeration/Windows/WindowsMMDeviceEnumeration/Public/WindowsMMNotificationClient.h
@@ -41,7 +41,7 @@ namespace Audio
 		WINDOWSMMDEVICEENUMERATION_API FString GetFriendlyName(const TComPtr<IMMDevice>& InDevice);
 
 		WINDOWSMMDEVICEENUMERATION_API TComPtr<IMMDevice> GetDevice(const FString InDeviceID) const;
-		WINDOWSMMDEVICEENUMERATION_API uint32 ReleaseClient() { return Release(); }
+		WINDOWSMMDEVICEENUMERATION_API uint32 ReleaseClient();
 
 		WINDOWSMMDEVICEENUMERATION_API void RegisterDeviceChangedListener(Audio::IAudioMixerDeviceChangedListener* DeviceChangedListener);
 		WINDOWSMMDEVICEENUMERATION_API void UnRegisterDeviceDeviceChangedListener(Audio::IAudioMixerDeviceChangedListener* DeviceChangedListener);
