Aggregate device enumeration and editor UI for audio device selection.

diff --git a/Engine/Plugins/Runtime/ADM/Source/ADMSpatialization/Private/ADMSpatializationModule.cpp b/Engine/Plugins/Runtime/ADM/Source/ADMSpatialization/Private/ADMSpatializationModule.cpp
--- a/Engine/Plugins/Runtime/ADM/Source/ADMSpatialization/Private/ADMSpatializationModule.cpp
+++ b/Engine/Plugins/Runtime/ADM/Source/ADMSpatialization/Private/ADMSpatializationModule.cpp
@@ -30,6 +30,10 @@ namespace UE::ADM::Spatialization
 		{
 			UE_LOG(LogADMSpatialization, Error, TEXT("Failed to find ADM Spatialization Developer Settings. Default IP not set for ADM Spatialization."));
 		}
+
+		// Enable aggregate devices so we can support large channel counts
+		constexpr bool bEnable = true;
+		FAudioDeviceManager::EnableAggregateDeviceSupport(bEnable);
 	}
 
 	void FModule::ShutdownModule()
diff --git a/Engine/Source/Editor/AudioEditor/AudioEditor.Build.cs b/Engine/Source/Editor/AudioEditor/AudioEditor.Build.cs
--- a/Engine/Source/Editor/AudioEditor/AudioEditor.Build.cs
+++ b/Engine/Source/Editor/AudioEditor/AudioEditor.Build.cs
@@ -18,7 +18,8 @@ public class AudioEditor : ModuleRules
 				"DeveloperSettings",
 				"UMGEditor",
 				"AudioExtensions",
-				"AudioLinkEngine"
+				"AudioLinkEngine",
+				"ToolWidgets"
 			}
 		);
 
diff --git a/Engine/Source/Editor/AudioEditor/Classes/AudioEditorSettings.h b/Engine/Source/Editor/AudioEditor/Classes/AudioEditorSettings.h
--- a/Engine/Source/Editor/AudioEditor/Classes/AudioEditorSettings.h
+++ b/Engine/Source/Editor/AudioEditor/Classes/AudioEditorSettings.h
@@ -2,14 +2,24 @@
 
 #pragma once
 
+#include "AudioDefines.h"
 #include "Engine/DeveloperSettings.h"
+#include "DetailCategoryBuilder.h"
+#include "DetailLayoutBuilder.h"
+#include "DetailWidgetRow.h"
+#include "IDetailCustomization.h"
 #include "HAL/Platform.h"
 #include "UObject/NameTypes.h"
 #include "UObject/ObjectMacros.h"
 #include "UObject/UObjectGlobals.h"
-#include "AudioDefines.h"
+
 #include "AudioEditorSettings.generated.h"
 
+namespace Audio
+{
+	struct FAudioPlatformDeviceInfo;
+}
+
 class UObject;
 struct FPropertyChangedEvent;
 
@@ -41,6 +51,14 @@ public:
 	UPROPERTY(EditAnywhere, config, Category = AssetMenu)
 	bool bPinSoundConcurrencyInAssetMenu = true;
 
+	/** Use current audio playback device selected in the operating system. */
+	UPROPERTY(EditAnywhere, config, Category = AudioOutputDevice, DisplayName = "Use System Default Audio Device")
+	bool bUseSystemDevice = true;
+	
+	/** The device id of the currently selected audio output device. */
+	UPROPERTY(EditAnywhere, config, Category = AudioOutputDevice)
+	FString AudioOutputDeviceId;
+
 	/** Set and apply, whether audio attenuation is used for non-game worlds*/
 	void SetUseAudioAttenuation(bool bInUseAudioAttenuation);
 
@@ -53,7 +71,7 @@ public:
 
 	//~ Begin UObject
 	virtual void PostEditChangeProperty(FPropertyChangedEvent& PropertyChangedEvent) override;
-	virtual void PostInitProperties();
+	virtual void PostInitProperties() override;
 	//~ End UObject
 
 private:
@@ -64,3 +82,38 @@ private:
 	/** Apply non-game world attenuation setting for AudioDevice with given ID.*/
 	void ApplyAttenuationForAudioDevice(Audio::FDeviceId InDeviceID);
 };
+
+class FAudioOutputDeviceCustomization : public IDetailCustomization
+{
+public:
+	static TSharedRef<IDetailCustomization> MakeInstance();
+
+	FAudioOutputDeviceCustomization();
+	virtual ~FAudioOutputDeviceCustomization() override;
+	
+	//~ Begin IDetailCustomization
+	virtual void CustomizeDetails(IDetailLayoutBuilder& DetailBuilder) override;
+	//~ End IDetailCustomization
+
+private:
+	TSharedPtr<IPropertyHandle> UseSystemDevicePropertyHandle;
+	TSharedPtr<IPropertyHandle> AudioDeviceIdPropertyHandle;
+	TSharedPtr<SWidget> DeviceListComboButton;
+	FString CurrentDeviceName;
+	TWeakPtr<FActiveTimerHandle> DeviceMenuActiveTimer;
+	FTickerDelegate TickDelegate;
+	FTSTicker::FDelegateHandle TickDelegateHandle;
+	
+	static TArray<Audio::FAudioPlatformDeviceInfo> GetAvailableAudioOutputDevices();
+	static Audio::FAudioPlatformDeviceInfo GetDeviceInfo(const FString& InDeviceId);
+	static TOptional<Audio::FAudioPlatformDeviceInfo> FindDefaultOutputDevice();
+	static FString GetDeviceNameForDeviceId(const FString& InDeviceId);
+	static FString GetCurrentAudioMixerDeviceName();
+	static void SetIsListeningForDeviceEvents(bool bInListeningForDeviceEvents);
+	static void RequestDeviceSwap(const FString& InDeviceId);
+	static bool IsAggregateHardwareDeviceId(const FString& InDeviceId);
+	
+	TSharedRef<SWidget> OnGenerateDeviceMenu();
+	void MenuItemDeviceSelected(Audio::FAudioPlatformDeviceInfo InDeviceInfo);
+	bool Tick(float DeltaTime);
+};
diff --git a/Engine/Source/Editor/AudioEditor/Private/AudioEditorModule.cpp b/Engine/Source/Editor/AudioEditor/Private/AudioEditorModule.cpp
--- a/Engine/Source/Editor/AudioEditor/Private/AudioEditorModule.cpp
+++ b/Engine/Source/Editor/AudioEditor/Private/AudioEditorModule.cpp
@@ -7,6 +7,7 @@
 #include "AssetTypeActions/AssetTypeActions_SoundClass.h"
 #include "AssetTypeActions/AssetTypeActions_SoundEffectPreset.h"
 #include "AssetTypeActions/AssetTypeActions_SoundSubmix.h"
+#include "AudioEditorSettings.h"
 #include "ClassTemplateEditorSubsystem.h"
 #include "ComponentAssetBroker.h"
 #include "Components/SynthComponent.h"
@@ -153,6 +154,8 @@ public:
 		AddCustomClass("SoundfieldEndpointSubmix", FOnGetDetailCustomizationInstance::CreateStatic(&FSoundfieldEndpointSubmixDetailsCustomization::MakeInstance));
 		AddCustomClass("SoundfieldSubmix", FOnGetDetailCustomizationInstance::CreateStatic(&FSoundfieldSubmixDetailsCustomization::MakeInstance));
 
+		AddCustomClass("AudioEditorSettings", FOnGetDetailCustomizationInstance::CreateStatic(&FAudioOutputDeviceCustomization::MakeInstance));
+
 		SetupIcons();
 
 		// Disable audio component asset broker in restricted mode 
@@ -380,6 +383,18 @@ public:
 		bIsRestricted = bInRestricted;
 	}
 
+	virtual TOptional<FAudioEditorDeviceSettings> GetAudioEditorDeviceSettings() const override
+	{
+		FAudioEditorDeviceSettings DeviceSettings;
+		if (const UAudioEditorSettings* AudioEditorSettings = GetDefault<UAudioEditorSettings>())
+		{
+			DeviceSettings.bUseSystemDevice = AudioEditorSettings->bUseSystemDevice;
+			DeviceSettings.DeviceId = AudioEditorSettings->AudioOutputDeviceId;
+		}
+		
+		return DeviceSettings;
+	}
+
 private:
 
 	void SetupIcons()
diff --git a/Engine/Source/Editor/AudioEditor/Private/AudioEditorSettings.cpp b/Engine/Source/Editor/AudioEditor/Private/AudioEditorSettings.cpp
--- a/Engine/Source/Editor/AudioEditor/Private/AudioEditorSettings.cpp
+++ b/Engine/Source/Editor/AudioEditor/Private/AudioEditorSettings.cpp
@@ -3,8 +3,17 @@
 #include "AudioEditorSettings.h"
 #include "AudioDeviceManager.h"
 #include "AudioDevice.h"
+#include "AudioMixerDevice.h"
+#include "DetailsNameWidgetOverrideCustomization.h"
+#include "Editor.h"
+#include "Framework/MultiBox/MultiBoxBuilder.h"
+#include "SSimpleComboButton.h"
+#include "Widgets/Input/SCheckBox.h"
+
+#define LOCTEXT_NAMESPACE "FAudioOutputDevicePropertyCustomization"
 
 struct FPropertyChangedEvent;
+using namespace Audio;
 
 void UAudioEditorSettings::PostInitProperties()
 {
@@ -57,3 +66,384 @@ void UAudioEditorSettings::ApplyAttenuationForAudioDevice(Audio::FDeviceId InDev
 		}
 	}	
 }
+
+TSharedRef<IDetailCustomization> FAudioOutputDeviceCustomization::MakeInstance()
+{
+	return MakeShareable(new FAudioOutputDeviceCustomization);
+}
+
+FAudioOutputDeviceCustomization::FAudioOutputDeviceCustomization()
+{
+	TickDelegate = FTickerDelegate::CreateRaw(this, &FAudioOutputDeviceCustomization::Tick);
+}
+
+FAudioOutputDeviceCustomization::~FAudioOutputDeviceCustomization()
+{
+	if (TickDelegateHandle.IsValid())
+	{
+		FTSTicker::GetCoreTicker().RemoveTicker(TickDelegateHandle);
+	}
+}
+
+bool FAudioOutputDeviceCustomization::Tick(float DeltaTime)
+{
+	bool bUseSystemDefault = false;
+	if (UseSystemDevicePropertyHandle->GetValue(bUseSystemDefault) == FPropertyAccess::Result::Success)
+	{
+		// The system default device can change when the user selects a new
+		// device via the OS settings. If the details are open, periodically check
+		// if the device name needs to be updated.
+		if (bUseSystemDefault)
+		{
+			FString NewDeviceName = GetCurrentAudioMixerDeviceName();
+			if (NewDeviceName != CurrentDeviceName)
+			{
+				CurrentDeviceName = NewDeviceName;
+				if (DeviceListComboButton.IsValid())
+				{
+					DeviceListComboButton->Invalidate(EInvalidateWidgetReason::Paint);
+				}
+			}
+		}
+	}
+	
+	return true;
+}
+
+void FAudioOutputDeviceCustomization::CustomizeDetails(IDetailLayoutBuilder& DetailBuilder)
+{
+	IDetailCategoryBuilder& Category = DetailBuilder.EditCategory(TEXT("AudioOutputDevice"));
+	IDetailLayoutBuilder& DetailLayout = Category.GetParentLayout();
+
+	UseSystemDevicePropertyHandle = DetailLayout.GetProperty(GET_MEMBER_NAME_CHECKED(UAudioEditorSettings, bUseSystemDevice));
+	IDetailPropertyRow& UseSystemPropertyRow = Category.AddProperty(UseSystemDevicePropertyHandle);
+
+	// Create check box which indicates system default audio device should be used (disables device combo button)
+	if (UseSystemDevicePropertyHandle.IsValid())
+	{
+		UseSystemPropertyRow.CustomWidget()
+		.NameContent()
+		[
+			UseSystemDevicePropertyHandle->CreatePropertyNameWidget()
+		]
+		.ValueContent()
+		[
+			SNew(SCheckBox)
+			.IsChecked_Lambda([this]() -> ECheckBoxState
+			{
+				bool bUseSystemDefault;
+				UseSystemDevicePropertyHandle->GetValue(bUseSystemDefault);
+
+				return bUseSystemDefault ? ECheckBoxState::Checked : ECheckBoxState::Unchecked;
+			})
+			.OnCheckStateChanged_Lambda([this](ECheckBoxState NewState) -> void
+			{
+				bool bUseSystemDefault = (NewState == ECheckBoxState::Checked);
+				UseSystemDevicePropertyHandle->SetValue(bUseSystemDefault);
+				SetIsListeningForDeviceEvents(bUseSystemDefault);
+
+				TOptional<FAudioPlatformDeviceInfo> DefaultDeviceInfo = FindDefaultOutputDevice();
+				if (DefaultDeviceInfo.IsSet())
+				{
+					if (bUseSystemDefault)
+					{
+						// This avoids a UI thrash due to the asynchronous nature of device swaps. It will
+						// take some time for the swap to complete so we set the UI name to that of the current
+						// active device here (which is not the same for aggregate devices).
+						CurrentDeviceName = GetCurrentAudioMixerDeviceName();
+						
+						// Swap to current system default device
+						RequestDeviceSwap(DefaultDeviceInfo->DeviceId);
+					}
+					else
+					{
+						// If we were using the system default device, and now we're not, lookup the current
+						// system default device and set that as the current device.
+						AudioDeviceIdPropertyHandle->SetValue(DefaultDeviceInfo->DeviceId);
+						CurrentDeviceName = DefaultDeviceInfo->Name;
+					}
+				}
+			})
+		];
+	}
+	
+	AudioDeviceIdPropertyHandle = DetailLayout.GetProperty(TEXT("AudioOutputDeviceId"));
+	IDetailPropertyRow& AudioDevicePropertyRow = Category.AddProperty(AudioDeviceIdPropertyHandle);
+
+	bool bUseSystemDefault;
+	UseSystemDevicePropertyHandle->GetValue(bUseSystemDefault);
+
+	// Populate initial value of CurrentDeviceName based depending on the current settings
+	if (bUseSystemDefault)
+	{
+		CurrentDeviceName = GetCurrentAudioMixerDeviceName();
+	}
+	else
+	{
+		FString CurrentDeviceId;
+		AudioDeviceIdPropertyHandle->GetValue(CurrentDeviceId);
+		CurrentDeviceName = GetDeviceNameForDeviceId(CurrentDeviceId);
+	}
+
+	AudioDevicePropertyRow.CustomWidget()
+	.NameContent()
+	[
+		SNew(STextBlock)
+		.Text(FText::FromString("Available Audio Output Devices"))
+		.Font(IDetailLayoutBuilder::GetDetailFont())
+	]
+	.ValueContent()
+	[
+		SNew(SSimpleComboButton)
+		.OnGetMenuContent(this, &FAudioOutputDeviceCustomization::OnGenerateDeviceMenu)
+		.ToolTipText(LOCTEXT("AudioOutputDeviceMenuToolTip", "Available Audio Output Devices"))
+		.HasDownArrow(true)
+		.UsesSmallText(true)
+		.IsEnabled_Lambda([this]()
+		{
+			bool bUseSystemDefault;
+			UseSystemDevicePropertyHandle->GetValue(bUseSystemDefault);
+			return !bUseSystemDefault;
+		})
+		.Text_Lambda([this]()
+		{
+			return FText::AsCultureInvariant(CurrentDeviceName);
+		})
+	];
+
+	TSharedPtr<SWidget> NameWidget;
+	AudioDevicePropertyRow.GetDefaultWidgets(NameWidget, DeviceListComboButton);
+    			
+	constexpr float TickDelay = 0.1f; // 100ms delay between ticks
+	if (!TickDelegateHandle.IsValid())
+	{
+		TickDelegateHandle = FTSTicker::GetCoreTicker().AddTicker(TickDelegate, TickDelay);
+	}
+}
+
+TArray<FAudioPlatformDeviceInfo> FAudioOutputDeviceCustomization::GetAvailableAudioOutputDevices()
+{
+	TArray<FAudioPlatformDeviceInfo> OutputDevices;
+
+	if (GEditor)
+	{
+		if (FMixerDevice* AudioMixerDevice = FAudioDeviceManager::GetAudioMixerDeviceFromWorldContext(GEditor->GetEditorWorldContext().World()))
+		{
+			if (Audio::IAudioMixerPlatformInterface* MixerPlatform = AudioMixerDevice->GetAudioMixerPlatform())
+			{
+				if (IAudioPlatformDeviceInfoCache* DeviceInfoCache = MixerPlatform->GetDeviceInfoCache())
+				{
+					OutputDevices = DeviceInfoCache->GetAllActiveOutputDevices();
+				}
+				else 
+				{
+					uint32 NumOutputDevices = 0;
+					MixerPlatform->GetNumOutputDevices(NumOutputDevices);
+					OutputDevices.Reserve(NumOutputDevices);
+					FAudioPlatformDeviceInfo CurrentOutputDevice = MixerPlatform->GetPlatformDeviceInfo();
+
+					for (uint32 i = 0; i < NumOutputDevices; ++i)
+					{
+						FAudioPlatformDeviceInfo DeviceInfo;
+						MixerPlatform->GetOutputDeviceInfo(i, DeviceInfo);
+						
+						OutputDevices.Emplace(MoveTemp(DeviceInfo));
+					}
+				}
+			}
+		}
+	}
+
+	return OutputDevices;
+}
+
+TOptional<FAudioPlatformDeviceInfo> FAudioOutputDeviceCustomization::FindDefaultOutputDevice()
+{
+	if (GEditor)
+	{
+		if (FMixerDevice* AudioMixerDevice = FAudioDeviceManager::GetAudioMixerDeviceFromWorldContext(GEditor->GetEditorWorldContext().World()))
+		{
+			if (Audio::IAudioMixerPlatformInterface* MixerPlatform = AudioMixerDevice->GetAudioMixerPlatform())
+			{
+				if (IAudioPlatformDeviceInfoCache* DeviceInfoCache = MixerPlatform->GetDeviceInfoCache())
+				{
+					return DeviceInfoCache->FindDefaultOutputDevice();
+				}
+			}
+		}
+	}
+
+	return{};
+}
+
+FAudioPlatformDeviceInfo FAudioOutputDeviceCustomization::GetDeviceInfo(const FString& InDeviceId)
+{
+	if (GEditor)
+	{
+		if (FMixerDevice* AudioMixerDevice = FAudioDeviceManager::GetAudioMixerDeviceFromWorldContext(GEditor->GetEditorWorldContext().World()))
+		{
+			if (Audio::IAudioMixerPlatformInterface* AudioMixerPlatform = AudioMixerDevice->GetAudioMixerPlatform())
+			{
+				uint32 NumOutputDevices = 0;
+				AudioMixerPlatform->GetNumOutputDevices(NumOutputDevices);
+
+				for (uint32 i = 0; i < NumOutputDevices; ++i)
+				{
+					FAudioPlatformDeviceInfo DeviceInfo;
+					AudioMixerPlatform->GetOutputDeviceInfo(i, DeviceInfo);
+
+					if (DeviceInfo.DeviceId == InDeviceId)
+					{
+						return DeviceInfo;
+					}
+				}
+			}
+		}
+	}
+
+	return {};
+}
+
+FString FAudioOutputDeviceCustomization::GetDeviceNameForDeviceId(const FString& InDeviceId)
+{
+	FAudioPlatformDeviceInfo DeviceInfo = GetDeviceInfo(InDeviceId);
+	return DeviceInfo.Name;
+}
+
+FString FAudioOutputDeviceCustomization::GetCurrentAudioMixerDeviceName()
+{
+	if (GEditor)
+	{
+		if (FMixerDevice* AudioMixerDevice = FAudioDeviceManager::GetAudioMixerDeviceFromWorldContext(GEditor->GetEditorWorldContext().World()))
+		{
+			if (Audio::IAudioMixerPlatformInterface* AudioMixerPlatform = AudioMixerDevice->GetAudioMixerPlatform())
+			{
+				return AudioMixerPlatform->GetCurrentDeviceName();
+			}
+		}
+	}
+
+	return {};
+}
+
+void FAudioOutputDeviceCustomization::SetIsListeningForDeviceEvents(bool bInListeningForDeviceEvents)
+{
+	if (GEditor)
+	{
+		if (FMixerDevice* AudioMixerDevice = FAudioDeviceManager::GetAudioMixerDeviceFromWorldContext(GEditor->GetEditorWorldContext().World()))
+		{
+			if (Audio::IAudioMixerPlatformInterface* AudioMixerPlatform = AudioMixerDevice->GetAudioMixerPlatform())
+			{
+				AudioMixerPlatform->SetIsListeningForDeviceEvents(bInListeningForDeviceEvents);
+			}
+		}
+	}
+}
+
+void FAudioOutputDeviceCustomization::RequestDeviceSwap(const FString& InDeviceId)
+{
+	if (GEditor)
+	{
+		if (FMixerDevice* AudioMixerDevice = FAudioDeviceManager::GetAudioMixerDeviceFromWorldContext(GEditor->GetEditorWorldContext().World()))
+		{
+			if (Audio::IAudioMixerPlatformInterface* AudioMixerPlatform = AudioMixerDevice->GetAudioMixerPlatform())
+			{
+				AudioMixerPlatform->RequestDeviceSwap(InDeviceId, /* force */true, TEXT("FAudioOutputDeviceCustomization::RequestDeviceSwap"));
+			}
+		}
+	}
+}
+
+bool FAudioOutputDeviceCustomization::IsAggregateHardwareDeviceId(const FString& InDeviceId)
+{
+	if (GEditor)
+	{
+		if (FMixerDevice* AudioMixerDevice = FAudioDeviceManager::GetAudioMixerDeviceFromWorldContext(GEditor->GetEditorWorldContext().World()))
+		{
+			if (Audio::IAudioMixerPlatformInterface* AudioMixerPlatform = AudioMixerDevice->GetAudioMixerPlatform())
+			{
+				if (IAudioPlatformDeviceInfoCache* DeviceInfoCache = AudioMixerPlatform->GetDeviceInfoCache())
+				{
+					return DeviceInfoCache->IsAggregateHardwareDeviceId(*InDeviceId);
+				}
+			}
+		}
+	}
+
+	return false;
+}
+
+TSharedRef<SWidget> FAudioOutputDeviceCustomization::OnGenerateDeviceMenu()
+{
+	FMenuBuilder MenuBuilder(true, nullptr, nullptr);
+
+	TArray<FAudioPlatformDeviceInfo> OutputDevices = GetAvailableAudioOutputDevices();
+	TArray<FAudioPlatformDeviceInfo> AggregateDevices;
+	TArray<FAudioPlatformDeviceInfo> NonAggregateDevices;
+
+	for (const FAudioPlatformDeviceInfo& DeviceInfo : OutputDevices)
+	{
+		if (IsAggregateHardwareDeviceId(DeviceInfo.DeviceId))
+		{
+			AggregateDevices.Add(DeviceInfo);
+		}
+		else
+		{
+			NonAggregateDevices.Add(DeviceInfo);
+		}
+	}
+
+	if (AggregateDevices.Num() > 0)
+	{
+		MenuBuilder.BeginSection(NAME_None, LOCTEXT("AggregateDeviceMenuSection", "Aggregate Audio Output Devices"));
+		{
+			for (const FAudioPlatformDeviceInfo& DeviceInfo : AggregateDevices)
+			{
+				if (!DeviceInfo.Name.IsEmpty())
+				{
+					MenuBuilder.AddMenuEntry(
+						FText::FromString(DeviceInfo.Name),
+						FText(),
+						FSlateIcon(),
+						FUIAction(
+							FExecuteAction::CreateSP(this, &FAudioOutputDeviceCustomization::MenuItemDeviceSelected, DeviceInfo)
+						)
+					);
+				}
+			}
+		}
+		MenuBuilder.EndSection();
+	}
+	
+	MenuBuilder.BeginSection(NAME_None, LOCTEXT("AudioOutputDeviceMenuSection", "Audio Output Devices"));
+	{
+		for (const FAudioPlatformDeviceInfo& DeviceInfo : NonAggregateDevices)
+		{
+			if (!DeviceInfo.Name.IsEmpty())
+			{
+				MenuBuilder.AddMenuEntry(
+					FText::FromString(DeviceInfo.Name),
+					FText(),
+					FSlateIcon(),
+					FUIAction(
+						FExecuteAction::CreateSP(this, &FAudioOutputDeviceCustomization::MenuItemDeviceSelected, DeviceInfo)
+					)
+				);
+			}
+		}
+	}
+	MenuBuilder.EndSection();
+
+	return MenuBuilder.MakeWidget();
+}
+
+void FAudioOutputDeviceCustomization::MenuItemDeviceSelected(FAudioPlatformDeviceInfo InDeviceInfo)
+{
+	CurrentDeviceName = InDeviceInfo.Name;
+	AudioDeviceIdPropertyHandle->SetValue(InDeviceInfo.DeviceId);
+
+	// User has changed device to be used by the editor, request a device swap
+	RequestDeviceSwap(InDeviceInfo.DeviceId);
+}
+
+#undef LOCTEXT_NAMESPACE
\ No newline at end of file
diff --git a/Engine/Source/Editor/AudioEditor/Public/AudioEditorModule.h b/Engine/Source/Editor/AudioEditor/Public/AudioEditorModule.h
--- a/Engine/Source/Editor/AudioEditor/Public/AudioEditorModule.h
+++ b/Engine/Source/Editor/AudioEditor/Public/AudioEditorModule.h
@@ -43,6 +43,12 @@ public:
 	virtual UObject* GetAssetFromComponent(UActorComponent* InComponent) override;
 };
 
+struct FAudioEditorDeviceSettings
+{
+	bool bUseSystemDevice = true;
+	FString DeviceId;
+};
+
 /** Sound class editor module interface */
 class IAudioEditorModule :	public IModuleInterface
 {
@@ -95,6 +101,9 @@ public:
 	// Sets if the Audio editor is in "restricted mode" (certain asset actions may be disabled)
 	virtual void SetRestrictedMode(bool bInRestricted) = 0;
 
+	// Returns the default audio editor device settings
+	virtual TOptional<FAudioEditorDeviceSettings> GetAudioEditorDeviceSettings() const = 0;
+
 	// Allows systems to register to open soundwave files with custom editor
 	DECLARE_DELEGATE_OneParam(SoundWaveEditorOpenDelegate, TArray<USoundWave*>);
 
diff --git a/Engine/Source/Runtime/AudioDeviceEnumeration/Windows/WindowsMMDeviceEnumeration/Private/WindowsMMCvarUtils.cpp b/Engine/Source/Runtime/AudioDeviceEnumeration/Windows/WindowsMMDeviceEnumeration/Private/WindowsMMCvarUtils.cpp
--- a/Engine/Source/Runtime/AudioDeviceEnumeration/Windows/WindowsMMDeviceEnumeration/Private/WindowsMMCvarUtils.cpp
+++ b/Engine/Source/Runtime/AudioDeviceEnumeration/Windows/WindowsMMDeviceEnumeration/Private/WindowsMMCvarUtils.cpp
@@ -19,6 +19,14 @@ FAutoConsoleVariableRef CVarEnableDetailedWindowsDeviceLogging(
 	TEXT("0: Not Enabled, 1: Enabled"),
 	ECVF_Default);
 
+static int32 EnableAggregateAudioDevicesCVar = 0;
+FAutoConsoleVariableRef CVarEnableAggregateAudioDevices(
+	TEXT("au.Wasapi.EnableAggregateAudioDevices"),
+	EnableAggregateAudioDevicesCVar,
+	TEXT("Enables WASAPI aggregate audio devices.\n")
+	TEXT("0: Not Enabled, 1: Enabled"),
+	ECVF_Default);
+
 namespace Audio
 {
 	bool WindowsMMCvarUtils::ShouldIgnoreDeviceSwaps()
@@ -31,4 +39,8 @@ namespace Audio
 		return EnableDetailedWindowsDeviceLoggingCVar != 0;
 	}
 
+	bool WindowsMMCvarUtils::IsAggregateDeviceSupportCVarEnabled()
+	{
+		return EnableAggregateAudioDevicesCVar != 0;
+	}
 } // namespace Audio
diff --git a/Engine/Source/Runtime/AudioDeviceEnumeration/Windows/WindowsMMDeviceEnumeration/Private/WindowsMMCvarUtils.h b/Engine/Source/Runtime/AudioDeviceEnumeration/Windows/WindowsMMDeviceEnumeration/Private/WindowsMMCvarUtils.h
--- a/Engine/Source/Runtime/AudioDeviceEnumeration/Windows/WindowsMMDeviceEnumeration/Private/WindowsMMCvarUtils.h
+++ b/Engine/Source/Runtime/AudioDeviceEnumeration/Windows/WindowsMMDeviceEnumeration/Private/WindowsMMCvarUtils.h
@@ -14,6 +14,8 @@ namespace Audio
 		/** Called by FWindowsMMNotificationClient to toggle logging for audio device changes: */
 		static bool ShouldLogDeviceSwaps();
 
+		/** Called by FWindowsMMNotificationClient to determine if Wasapi aggregate device support is enabled. */
+		static bool IsAggregateDeviceSupportCVarEnabled();
 	};
 
 } // namespace Audio
\ No newline at end of file
diff --git a/Engine/Source/Runtime/AudioDeviceEnumeration/Windows/WindowsMMDeviceEnumeration/Private/WindowsMMDeviceInfoCache.cpp b/Engine/Source/Runtime/AudioDeviceEnumeration/Windows/WindowsMMDeviceEnumeration/Private/WindowsMMDeviceInfoCache.cpp
--- a/Engine/Source/Runtime/AudioDeviceEnumeration/Windows/WindowsMMDeviceEnumeration/Private/WindowsMMDeviceInfoCache.cpp
+++ b/Engine/Source/Runtime/AudioDeviceEnumeration/Windows/WindowsMMDeviceEnumeration/Private/WindowsMMDeviceInfoCache.cpp
@@ -2,751 +2,946 @@
 
 #include "WindowsMMDeviceInfoCache.h"
 
-THIRD_PARTY_INCLUDES_START
-#include <devicetopology.h>
-#include <functiondiscoverykeys_devpkey.h>
-THIRD_PARTY_INCLUDES_END
-
 #include "ConversionHelpers.h"
 #include "HAL/IConsoleManager.h"
+#include "Internationalization/Regex.h"
 #include "Misc/CommandLine.h"
 #include "Misc/Parse.h"
-#include "Misc/ScopeRWLock.h"
+
 #include "Windows/AllowWindowsPlatformTypes.h"
 #include "Windows/AllowWindowsPlatformAtomics.h"
+
+THIRD_PARTY_INCLUDES_START
+#include <devicetopology.h>
+#include <functiondiscoverykeys_devpkey.h>
+THIRD_PARTY_INCLUDES_END
+
 #include "Windows/HideWindowsPlatformAtomics.h"
 #include "Windows/HideWindowsPlatformTypes.h"
+
 #include "WindowsMMStringUtils.h"
+#include "Misc/ScopeRWLock.h"
 #include "WindowsMMCvarUtils.h"
 #include "WindowsMMDeviceEnumerationLog.h"
 
-
 namespace Audio
 {
+	// Support for forcing the audio renderer on windows machines to initialize as if it was connected to a 7.1 output device.
+	// This is useful for cases where a developer does not have access to a 7.1 device (or is working on a cloud machine) and 
+	// wants to validate and excersie surround sound audio rendering code.
+	static bool GetForceSurroundSound()
+	{	
+		static bool bForceSurroundSound = FParse::Param(FCommandLine::Get(), TEXT("ForceSurroundSound"));
+		return bForceSurroundSound;
+	}
+	
+	FWindowsMMDeviceCache::FCacheEntry& FWindowsMMDeviceCache::FCacheEntry::operator=(FCacheEntry&& InOther)
+	{
+		DeviceId = MoveTemp(InOther.DeviceId);
+		FriendlyName = MoveTemp(InOther.FriendlyName);
+		DeviceFriendlyName = MoveTemp(InOther.DeviceFriendlyName);
+		State = MoveTemp(InOther.State);
+		NumChannels = MoveTemp(InOther.NumChannels);
+		SampleRate = MoveTemp(InOther.SampleRate);
+		Type = MoveTemp(InOther.Type);
+		ChannelBitmask = MoveTemp(InOther.ChannelBitmask);
+		OutputChannels = MoveTemp(InOther.OutputChannels);
+		HardwareId = MoveTemp(InOther.HardwareId);
+		FilterId = MoveTemp(InOther.FilterId);
+		return *this;
+	}
 
-// Support for forcing the audio renderer on windows machines to initialize as if it was connected to a 7.1 output device.
-// This is useful for cases where a developer does not have access to a 7.1 device (or is working on a cloud machine) and 
-// wants to validate and excersie surround sound audio rendering code.
-static bool GetForceSurroundSound()
-{	
-	static bool bForceSurroundSound = FParse::Param(FCommandLine::Get(), TEXT("ForceSurroundSound"));
-	return bForceSurroundSound;
-}
-
-FWindowsMMDeviceCache::FCacheEntry& FWindowsMMDeviceCache::FCacheEntry::operator=(FCacheEntry&& InOther)
-{
-	DeviceId = MoveTemp(InOther.DeviceId);
-	FriendlyName = MoveTemp(InOther.FriendlyName);
-	DeviceFriendlyName = MoveTemp(InOther.DeviceFriendlyName);
-	State = MoveTemp(InOther.State);
-	NumChannels = MoveTemp(InOther.NumChannels);
-	SampleRate = MoveTemp(InOther.SampleRate);
-	Type = MoveTemp(InOther.Type);
-	ChannelBitmask = MoveTemp(InOther.ChannelBitmask);
-	OutputChannels = MoveTemp(InOther.OutputChannels);
-	HardwareId = MoveTemp(InOther.HardwareId);
-	FilterId = MoveTemp(InOther.FilterId);
-	return *this;
-}
-
-FWindowsMMDeviceCache::FCacheEntry& FWindowsMMDeviceCache::FCacheEntry::operator=(const FCacheEntry& InOther)
-{
-	// Copy everything but the lock. 
-	DeviceId = InOther.DeviceId;
-	FriendlyName = InOther.FriendlyName;
-	DeviceFriendlyName = InOther.DeviceFriendlyName;
-	State = InOther.State;
-	NumChannels = InOther.NumChannels;
-	SampleRate = InOther.SampleRate;
-	Type = InOther.Type;
-	ChannelBitmask = InOther.ChannelBitmask;
-	OutputChannels = InOther.OutputChannels;
-	HardwareId = InOther.HardwareId;
-	FilterId = InOther.FilterId;
-	return *this;
-}
-
-FWindowsMMDeviceCache::FCacheEntry::FCacheEntry(const FString& InDeviceId)
-	: DeviceId{ InDeviceId }
-{
-}
-
-FWindowsMMDeviceCache::FCacheEntry::FCacheEntry(FCacheEntry&& InOther)
-{
-	*this = MoveTemp(InOther);
-}
-
-FWindowsMMDeviceCache::FCacheEntry::FCacheEntry(const FCacheEntry& InOther)
-{
-	*this = InOther;
-}
+	FWindowsMMDeviceCache::FCacheEntry& FWindowsMMDeviceCache::FCacheEntry::operator=(const FCacheEntry& InOther)
+	{
+		// Copy everything but the lock. 
+		DeviceId = InOther.DeviceId;
+		FriendlyName = InOther.FriendlyName;
+		DeviceFriendlyName = InOther.DeviceFriendlyName;
+		State = InOther.State;
+		NumChannels = InOther.NumChannels;
+		SampleRate = InOther.SampleRate;
+		Type = InOther.Type;
+		ChannelBitmask = InOther.ChannelBitmask;
+		OutputChannels = InOther.OutputChannels;
+		HardwareId = InOther.HardwareId;
+		FilterId = InOther.FilterId;
+		return *this;
+	}
 
-FWindowsMMDeviceCache::FWindowsMMDeviceCache(bool bInEnumerateHardware)
-	: bEnumerateHardware(bInEnumerateHardware)
-{
-	ensure(SUCCEEDED(CoCreateInstance(__uuidof(MMDeviceEnumerator), NULL, CLSCTX_INPROC_SERVER, IID_PPV_ARGS(&DeviceEnumerator))) && DeviceEnumerator);
+	FWindowsMMDeviceCache::FCacheEntry::FCacheEntry(const FString& InDeviceId)
+		: DeviceId{ InDeviceId }
+	{
+	}
 
-	EnumerateEndpoints();
-	EnumerateDefaults();
-}
+	FWindowsMMDeviceCache::FCacheEntry::FCacheEntry(FCacheEntry&& InOther)
+	{
+		*this = MoveTemp(InOther);
+	}
 
-bool FWindowsMMDeviceCache::EnumerateChannelMask(uint32 InMask, FCacheEntry& OutInfo)
-{
-	// Loop through the extensible format channel flags in the standard order and build our output channel array
-	// From https://msdn.microsoft.com/en-us/library/windows/hardware/dn653308(v=vs.85).aspx
-	// The channels in the interleaved stream corresponding to these spatial positions must appear in the order specified above. This holds true even in the 
-	// case of a non-contiguous subset of channels. For example, if a stream contains left, bass enhance and right, then channel 1 is left, channel 2 is right, 
-	// and channel 3 is bass enhance. This enables the linkage of multi-channel streams to well-defined multi-speaker configurations.
-
-	static const uint32 REMOVE_ME_ChannelTypeMap[EAudioMixerChannel::ChannelTypeCount] =
-	{
-		SPEAKER_FRONT_LEFT,
-		SPEAKER_FRONT_RIGHT,
-		SPEAKER_FRONT_CENTER,
-		SPEAKER_LOW_FREQUENCY,
-		SPEAKER_BACK_LEFT,
-		SPEAKER_BACK_RIGHT,
-		SPEAKER_FRONT_LEFT_OF_CENTER,
-		SPEAKER_FRONT_RIGHT_OF_CENTER,
-		SPEAKER_BACK_CENTER,
-		SPEAKER_SIDE_LEFT,
-		SPEAKER_SIDE_RIGHT,
-		SPEAKER_TOP_CENTER,
-		SPEAKER_TOP_FRONT_LEFT,
-		SPEAKER_TOP_FRONT_CENTER,
-		SPEAKER_TOP_FRONT_RIGHT,
-		SPEAKER_TOP_BACK_LEFT,
-		SPEAKER_TOP_BACK_CENTER,
-		SPEAKER_TOP_BACK_RIGHT,
-		SPEAKER_RESERVED,
-	};
-
-	OutInfo.ChannelBitmask = InMask;
-	OutInfo.OutputChannels.Reset();
-
-	// No need to enumerate speakers for capture devices.
-	if (OutInfo.Type == FCacheEntry::EEndpointType::Capture)
+	FWindowsMMDeviceCache::FCacheEntry::FCacheEntry(const FCacheEntry& InOther)
 	{
-		return true;
+		*this = InOther;
 	}
 
-	uint32 ChanCount = 0;
-	for (uint32 ChannelTypeIndex = 0; ChannelTypeIndex < EAudioMixerChannel::ChannelTypeCount && ChanCount < (uint32)OutInfo.NumChannels; ++ChannelTypeIndex)
+	FWindowsMMDeviceCache::FWindowsMMDeviceCache(bool bInEnableAggregateDeviceSupport) :
+		bIsAggregateDeviceSupportEnabled(bInEnableAggregateDeviceSupport)
 	{
-		if (InMask & REMOVE_ME_ChannelTypeMap[ChannelTypeIndex])
-		{
-			OutInfo.OutputChannels.Add((EAudioMixerChannel::Type)ChannelTypeIndex);
-			++ChanCount;
-		}
+		ensure(SUCCEEDED(CoCreateInstance(__uuidof(MMDeviceEnumerator), NULL, CLSCTX_INPROC_SERVER, IID_PPV_ARGS(&DeviceEnumerator))) && DeviceEnumerator);
+		
+		EnumerateEndpoints();
+		EnumerateDefaults();
 	}
 
-	// We didn't match channel masks for all channels, revert to a default ordering
-	if (ChanCount < (uint32)OutInfo.NumChannels)
+	FWindowsMMDeviceCache::FWindowsMMDeviceCache() :
+		FWindowsMMDeviceCache(false)
 	{
-		UE_CLOG(WindowsMMCvarUtils::ShouldLogDeviceSwaps(), LogAudioEnumeration, Warning, TEXT("FWindowsMMDeviceCache: Did not find the channel type flags for audio device '%s'. Reverting to a default channel ordering."), *OutInfo.FriendlyName);
+	}
 
-		OutInfo.OutputChannels.Reset();
 
-		static const EAudioMixerChannel::Type DefaultChannelOrdering[] = {
-			EAudioMixerChannel::FrontLeft,
-			EAudioMixerChannel::FrontRight,
-			EAudioMixerChannel::FrontCenter,
-			EAudioMixerChannel::LowFrequency,
-			EAudioMixerChannel::SideLeft,
-			EAudioMixerChannel::SideRight,
-			EAudioMixerChannel::BackLeft,
-			EAudioMixerChannel::BackRight,
+	bool FWindowsMMDeviceCache::EnumerateChannelMask(uint32 InMask, FCacheEntry& OutInfo)
+	{
+		// Loop through the extensible format channel flags in the standard order and build our output channel array
+		// From https://msdn.microsoft.com/en-us/library/windows/hardware/dn653308(v=vs.85).aspx
+		// The channels in the interleaved stream corresponding to these spatial positions must appear in the order specified above. This holds true even in the 
+		// case of a non-contiguous subset of channels. For example, if a stream contains left, bass enhance and right, then channel 1 is left, channel 2 is right, 
+		// and channel 3 is bass enhance. This enables the linkage of multi-channel streams to well-defined multi-speaker configurations.
+
+		static const uint32 REMOVE_ME_ChannelTypeMap[EAudioMixerChannel::ChannelTypeCount] =
+		{
+			SPEAKER_FRONT_LEFT,
+			SPEAKER_FRONT_RIGHT,
+			SPEAKER_FRONT_CENTER,
+			SPEAKER_LOW_FREQUENCY,
+			SPEAKER_BACK_LEFT,
+			SPEAKER_BACK_RIGHT,
+			SPEAKER_FRONT_LEFT_OF_CENTER,
+			SPEAKER_FRONT_RIGHT_OF_CENTER,
+			SPEAKER_BACK_CENTER,
+			SPEAKER_SIDE_LEFT,
+			SPEAKER_SIDE_RIGHT,
+			SPEAKER_TOP_CENTER,
+			SPEAKER_TOP_FRONT_LEFT,
+			SPEAKER_TOP_FRONT_CENTER,
+			SPEAKER_TOP_FRONT_RIGHT,
+			SPEAKER_TOP_BACK_LEFT,
+			SPEAKER_TOP_BACK_CENTER,
+			SPEAKER_TOP_BACK_RIGHT,
+			SPEAKER_RESERVED,
 		};
 
-		const EAudioMixerChannel::Type* ChannelOrdering = DefaultChannelOrdering;
+		OutInfo.ChannelBitmask = InMask;
+		OutInfo.OutputChannels.Reset();
 
-		// Override channel ordering for some special cases
-		if (OutInfo.NumChannels == 4)
+		// No need to enumerate speakers for capture devices.
+		if (OutInfo.Type == EDeviceEndpointType::Capture)
 		{
-			static EAudioMixerChannel::Type DefaultChannelOrderingQuad[] = {
-				EAudioMixerChannel::FrontLeft,
-				EAudioMixerChannel::FrontRight,
-				EAudioMixerChannel::BackLeft,
-				EAudioMixerChannel::BackRight,
-			};
+			return true;
+		}
 
-			ChannelOrdering = DefaultChannelOrderingQuad;
+		uint32 ChanCount = 0;
+		for (uint32 ChannelTypeIndex = 0; ChannelTypeIndex < EAudioMixerChannel::ChannelTypeCount && ChanCount < (uint32)OutInfo.NumChannels; ++ChannelTypeIndex)
+		{
+			if (InMask & REMOVE_ME_ChannelTypeMap[ChannelTypeIndex])
+			{
+				OutInfo.OutputChannels.Add((EAudioMixerChannel::Type)ChannelTypeIndex);
+				++ChanCount;
+			}
 		}
-		else if (OutInfo.NumChannels == 6)
+
+		// We didn't match channel masks for all channels, revert to a default ordering
+		if (ChanCount < (uint32)OutInfo.NumChannels)
 		{
-			static const EAudioMixerChannel::Type DefaultChannelOrdering51[] = {
+			UE_CLOG(WindowsMMCvarUtils::ShouldLogDeviceSwaps(), LogAudioEnumeration, Warning, TEXT("FWindowsMMDeviceCache: Did not find the channel type flags for audio device '%s'. Reverting to a default channel ordering."), *OutInfo.FriendlyName);
+
+			OutInfo.OutputChannels.Reset();
+
+			static const EAudioMixerChannel::Type DefaultChannelOrdering[] = {
 				EAudioMixerChannel::FrontLeft,
 				EAudioMixerChannel::FrontRight,
 				EAudioMixerChannel::FrontCenter,
 				EAudioMixerChannel::LowFrequency,
+				EAudioMixerChannel::SideLeft,
+				EAudioMixerChannel::SideRight,
 				EAudioMixerChannel::BackLeft,
 				EAudioMixerChannel::BackRight,
 			};
 
-			ChannelOrdering = DefaultChannelOrdering51;
-		}
+			const EAudioMixerChannel::Type* ChannelOrdering = DefaultChannelOrdering;
 
-		check(OutInfo.NumChannels <= 8);
-		for (int32 Index = 0; Index < OutInfo.NumChannels; ++Index)
-		{
-			OutInfo.OutputChannels.Add(ChannelOrdering[Index]);
+			// Override channel ordering for some special cases
+			if (OutInfo.NumChannels == 4)
+			{
+				static EAudioMixerChannel::Type DefaultChannelOrderingQuad[] = {
+					EAudioMixerChannel::FrontLeft,
+					EAudioMixerChannel::FrontRight,
+					EAudioMixerChannel::BackLeft,
+					EAudioMixerChannel::BackRight,
+				};
+
+				ChannelOrdering = DefaultChannelOrderingQuad;
+			}
+			else if (OutInfo.NumChannels == 6)
+			{
+				static const EAudioMixerChannel::Type DefaultChannelOrdering51[] = {
+					EAudioMixerChannel::FrontLeft,
+					EAudioMixerChannel::FrontRight,
+					EAudioMixerChannel::FrontCenter,
+					EAudioMixerChannel::LowFrequency,
+					EAudioMixerChannel::BackLeft,
+					EAudioMixerChannel::BackRight,
+				};
+
+				ChannelOrdering = DefaultChannelOrdering51;
+			}
+
+			check(OutInfo.NumChannels <= 8);
+			for (int32 Index = 0; Index < OutInfo.NumChannels; ++Index)
+			{
+				OutInfo.OutputChannels.Add(ChannelOrdering[Index]);
+			}
 		}
+		return true;
 	}
-	return true;
-}
-
-bool FWindowsMMDeviceCache::EnumerateChannelFormat(const WAVEFORMATEX* InFormat, FCacheEntry& OutInfo)
-{
-	OutInfo.OutputChannels.Empty();
 
-	// Extensible format supports surround sound so we need to parse the channel configuration to build our channel output array
-	if (InFormat->wFormatTag == WAVE_FORMAT_EXTENSIBLE)
+	bool FWindowsMMDeviceCache::EnumerateChannelFormat(const WAVEFORMATEX* InFormat, FCacheEntry& OutInfo)
 	{
-		// Cast to the extensible format to get access to extensible data
-		const WAVEFORMATEXTENSIBLE* WaveFormatExtensible = (const WAVEFORMATEXTENSIBLE*)(InFormat);
-		return EnumerateChannelMask(WaveFormatExtensible->dwChannelMask, OutInfo);
-	}
-	else
-	{
-		// Non-extensible formats only support mono or stereo channel output
-		OutInfo.OutputChannels.Add(EAudioMixerChannel::FrontLeft);
-		if (OutInfo.NumChannels == 2)
+		OutInfo.OutputChannels.Empty();
+
+		// Extensible format supports surround sound so we need to parse the channel configuration to build our channel output array
+		if (InFormat->wFormatTag == WAVE_FORMAT_EXTENSIBLE)
 		{
-			OutInfo.OutputChannels.Add(EAudioMixerChannel::FrontRight);
+			// Cast to the extensible format to get access to extensible data
+			const WAVEFORMATEXTENSIBLE* WaveFormatExtensible = (const WAVEFORMATEXTENSIBLE*)(InFormat);
+			return EnumerateChannelMask(WaveFormatExtensible->dwChannelMask, OutInfo);
+		}
+		else
+		{
+			// Non-extensible formats only support mono or stereo channel output
+			OutInfo.OutputChannels.Add(EAudioMixerChannel::FrontLeft);
+			if (OutInfo.NumChannels == 2)
+			{
+				OutInfo.OutputChannels.Add(EAudioMixerChannel::FrontRight);
+			}
 		}
-	}
 
-	// Aways success for now.
-	return true;
-}
+		// Aways success for now.
+		return true;
+	}
 
-FWindowsMMDeviceCache::FCacheEntry::EEndpointType FWindowsMMDeviceCache::QueryDeviceDataFlow(const TComPtr<IMMDevice>& InDevice) const
-{
-	TComPtr<IMMEndpoint> Endpoint;
-	if (SUCCEEDED(InDevice->QueryInterface(IID_PPV_ARGS(&Endpoint))))
+	EDeviceEndpointType FWindowsMMDeviceCache::QueryDeviceDataFlow(const TComPtr<IMMDevice>& InDevice) const
 	{
-		EDataFlow DataFlow = eRender;
-		if (SUCCEEDED(Endpoint->GetDataFlow(&DataFlow)))
+		TComPtr<IMMEndpoint> Endpoint;
+		if (SUCCEEDED(InDevice->QueryInterface(IID_PPV_ARGS(&Endpoint))))
 		{
-			switch (DataFlow)
+			EDataFlow DataFlow = eRender;
+			if (SUCCEEDED(Endpoint->GetDataFlow(&DataFlow)))
 			{
-			case eRender:
-				return FCacheEntry::EEndpointType::Render;
-			case eCapture:
-				return FCacheEntry::EEndpointType::Capture;
-			default:
-				break;
+				switch (DataFlow)
+				{
+				case eRender:
+					return EDeviceEndpointType::Render;
+				case eCapture:
+					return EDeviceEndpointType::Capture;
+				default:
+					break;
+				}
 			}
 		}
+		return EDeviceEndpointType::Unknown;
 	}
-	return FCacheEntry::EEndpointType::Unknown;
-}
 
-bool FWindowsMMDeviceCache::EnumerateDeviceProps(const TComPtr<IMMDevice>& InDevice, FCacheEntry& OutInfo)
-{
-	// Mark if this is a Render Device or Capture or Unknown.
-	OutInfo.Type = QueryDeviceDataFlow(InDevice);
-
-	// Also query the device state.
-	DWORD DeviceState = DEVICE_STATE_NOTPRESENT;
-	if (SUCCEEDED(InDevice->GetState(&DeviceState)))
+	bool FWindowsMMDeviceCache::EnumerateDeviceProps(const TComPtr<IMMDevice>& InDevice, FCacheEntry& OutInfo)
 	{
-		OutInfo.State = ConvertWordToDeviceState(DeviceState);
-	}
+		// Mark if this is a Render Device or Capture or Unknown.
+		OutInfo.Type = QueryDeviceDataFlow(InDevice);
 
-	TComPtr<IPropertyStore> PropertyStore;
-	if (SUCCEEDED(InDevice->OpenPropertyStore(STGM_READ, &PropertyStore)))
-	{
-		// Friendly Name
-		PROPVARIANT FriendlyName;
-		PropVariantInit(&FriendlyName);
-		if (SUCCEEDED(PropertyStore->GetValue(PKEY_Device_FriendlyName, &FriendlyName)) && FriendlyName.pwszVal)
+		// Also query the device state.
+		DWORD DeviceState = DEVICE_STATE_NOTPRESENT;
+		if (SUCCEEDED(InDevice->GetState(&DeviceState)))
 		{
-			OutInfo.FriendlyName = FString(FriendlyName.pwszVal);
-			PropVariantClear(&FriendlyName);
+			OutInfo.State = ConvertWordToDeviceState(DeviceState);
 		}
 
-		auto EnumDeviceFormat = [this](const TComPtr<IPropertyStore>& InPropStore, REFPROPERTYKEY InKey, FCacheEntry& OutInfo) -> bool
+		TComPtr<IPropertyStore> PropertyStore;
+		if (SUCCEEDED(InDevice->OpenPropertyStore(STGM_READ, &PropertyStore)))
 		{
-			// Device Format
-			PROPVARIANT DeviceFormat;
-			PropVariantInit(&DeviceFormat);
+			// Friendly Name
+			PROPVARIANT FriendlyName;
+			PropVariantInit(&FriendlyName);
+			if (SUCCEEDED(PropertyStore->GetValue(PKEY_Device_FriendlyName, &FriendlyName)) && FriendlyName.pwszVal)
+			{
+				OutInfo.FriendlyName = FString(FriendlyName.pwszVal);
+				PropVariantClear(&FriendlyName);
+			}
 
-			if (SUCCEEDED(InPropStore->GetValue(InKey, &DeviceFormat)) && DeviceFormat.blob.pBlobData)
+			auto EnumDeviceFormat = [this](const TComPtr<IPropertyStore>& InPropStore, REFPROPERTYKEY InKey, FCacheEntry& OutInfo) -> bool
 			{
-				const WAVEFORMATEX* WaveFormatEx = (const WAVEFORMATEX*)(DeviceFormat.blob.pBlobData);
-				if (GetForceSurroundSound())
-				{
-					OutInfo.NumChannels = 8;
-				}
-				else
+				// Device Format
+				PROPVARIANT DeviceFormat;
+				PropVariantInit(&DeviceFormat);
+
+				if (SUCCEEDED(InPropStore->GetValue(InKey, &DeviceFormat)) && DeviceFormat.blob.pBlobData)
 				{
-					OutInfo.NumChannels = FMath::Clamp((int32)WaveFormatEx->nChannels, 2, 8);
+					const WAVEFORMATEX* WaveFormatEx = (const WAVEFORMATEX*)(DeviceFormat.blob.pBlobData);
+					if (GetForceSurroundSound())
+					{
+						OutInfo.NumChannels = 8;
+					}
+					else
+					{
+						OutInfo.NumChannels = FMath::Clamp((int32)WaveFormatEx->nChannels, 2, 8);
+					}
+					
+					OutInfo.SampleRate = WaveFormatEx->nSamplesPerSec;
+
+					EnumerateChannelFormat(WaveFormatEx, OutInfo);
+					PropVariantClear(&DeviceFormat);
+					return true;
 				}
-				OutInfo.SampleRate = WaveFormatEx->nSamplesPerSec;
+				return false;
+			};
 
-				EnumerateChannelFormat(WaveFormatEx, OutInfo);
-				PropVariantClear(&DeviceFormat);
-				return true;
+			if (EnumDeviceFormat(PropertyStore, PKEY_AudioEngine_DeviceFormat, OutInfo) ||
+				EnumDeviceFormat(PropertyStore, PKEY_AudioEngine_OEMFormat, OutInfo))
+			{
 			}
-			return false;
-		};
+			else
+			{
+				// Log a warning if this device is active as we failed to ask for a format
+				UE_CLOG(DeviceState == DEVICE_STATE_ACTIVE, LogAudioEnumeration, Warning, TEXT("FWindowsMMDeviceCache: Failed to get Format for active device '%s'"), *OutInfo.FriendlyName);
+			}
+		}
+
+		// Aways success for now.
+		return true;
+	}
+
+	bool FWindowsMMDeviceCache::EnumerateHardwareTopology(const TComPtr<IMMDevice>& InDevice, FCacheEntry& OutInfo)
+	{
+		SCOPED_NAMED_EVENT(FWindowsMMDeviceCache_EnumerateHardwareTopology, FColor::Blue);
 
-		if (EnumDeviceFormat(PropertyStore, PKEY_AudioEngine_DeviceFormat, OutInfo) ||
-			EnumDeviceFormat(PropertyStore, PKEY_AudioEngine_OEMFormat, OutInfo))
+	#if PLATFORM_WINDOWS
+
+		if (!DeviceEnumerator)
 		{
+			return false;
 		}
-		else
+
+		TComPtr<IDeviceTopology> RenderEndpointTopology;
+		HRESULT Result = InDevice->Activate(__uuidof(IDeviceTopology), CLSCTX_ALL, NULL, (void**)&RenderEndpointTopology);
+		if (FAILED(Result))
 		{
-			// Log a warning if this device is active as we failed to ask for a format
-			UE_CLOG(DeviceState == DEVICE_STATE_ACTIVE, LogAudioEnumeration, Warning, TEXT("FWindowsMMDeviceCache: Failed to get Format for active device '%s'"), *OutInfo.FriendlyName);
+			UE_LOG(LogAudioEnumeration, Verbose, TEXT("FWindowsMMDeviceCache::EnumerateHardwareTopology failed to fetch IDeviceTopology: 0x%x"), Result);
+			return false;
 		}
-	}
 
-	// Aways success for now.
-	return true;
-}
+		TComPtr<IConnector> Connector;
+		Result = RenderEndpointTopology->GetConnector(0, &Connector);
+		if (FAILED(Result))
+		{
+			UE_LOG(LogAudioEnumeration, Verbose, TEXT("FWindowsMMDeviceCache::EnumerateHardwareTopology failed to fetch connector: 0x%x"), Result);
+			return false;
+		}
 
-bool FWindowsMMDeviceCache::EnumerateHardwareTopology(const TComPtr<IMMDevice>& InDevice, FCacheEntry& OutInfo)
-{
-	SCOPED_NAMED_EVENT(FWindowsMMDeviceCache_EnumerateHardwareTopology, FColor::Blue);
+		LPWSTR RenderFilterId = nullptr;
+		Result = Connector->GetDeviceIdConnectedTo(&RenderFilterId);
+		if (FAILED(Result) || !RenderFilterId)
+		{
+			UE_LOG(LogAudioEnumeration, Verbose, TEXT("FWindowsMMDeviceCache::EnumerateHardwareTopology failed to fetch render filter Id: 0x%x"), Result);
+			return false;
+		}
 
-#if PLATFORM_WINDOWS
+		FString TempFilterId(RenderFilterId);
+		CoTaskMemFree(RenderFilterId);
 
-	if (!DeviceEnumerator)
-	{
-		return false;
-	}
+		TComPtr<IMMDevice> RenderDevnode;
+		Result = DeviceEnumerator->GetDevice(*TempFilterId, &RenderDevnode);
+		if (FAILED(Result))
+		{
+			UE_LOG(LogAudioEnumeration, Verbose, TEXT("FWindowsMMDeviceCache::EnumerateHardwareTopology failed to fetch render node: 0x%x"), Result);
+			return false;
+		}
 
-	TComPtr<IDeviceTopology> RenderEndpointTopology;
-	HRESULT Result = InDevice->Activate(__uuidof(IDeviceTopology), CLSCTX_ALL, NULL, (void**)&RenderEndpointTopology);
-	if (FAILED(Result))
-	{
-		UE_LOG(LogAudioEnumeration, Verbose, TEXT("FWindowsMMDeviceCache::EnumerateHardwareTopology failed to fetch IDeviceTopology: 0x%x"), Result);
-		return false;
-	}
+		TComPtr<IPropertyStore> PropertyStore;
+		Result = RenderDevnode->OpenPropertyStore(STGM_READ, &PropertyStore);
+		if (FAILED(Result))
+		{
+			UE_LOG(LogAudioEnumeration, Verbose, TEXT("FWindowsMMDeviceCache::EnumerateHardwareTopology failed to open property store: 0x%x"), Result);
+			return false;
+		}
 
-	TComPtr<IConnector> Connector;
-	Result = RenderEndpointTopology->GetConnector(0, &Connector);
-	if (FAILED(Result))
-	{
-		UE_LOG(LogAudioEnumeration, Verbose, TEXT("FWindowsMMDeviceCache::EnumerateHardwareTopology failed to fetch connector: 0x%x"), Result);
-		return false;
-	}
+		PROPVARIANT HardwareId;
+		PropVariantInit(&HardwareId);
+		Result = PropertyStore->GetValue(PKEY_Device_InstanceId, &HardwareId);
+		if (FAILED(Result) || !HardwareId.pwszVal)
+		{
+			UE_LOG(LogAudioEnumeration, Verbose, TEXT("FWindowsMMDeviceCache::EnumerateHardwareTopology failed to fetch hardward Id: 0x%x"), Result);
+			return false;
+		}
 
-	LPWSTR RenderFilterId = nullptr;
-	Result = Connector->GetDeviceIdConnectedTo(&RenderFilterId);
-	if (FAILED(Result) || !RenderFilterId)
-	{
-		UE_LOG(LogAudioEnumeration, Verbose, TEXT("FWindowsMMDeviceCache::EnumerateHardwareTopology failed to fetch render filter Id: 0x%x"), Result);
-		return false;
-	}
+		OutInfo.HardwareId = FName(HardwareId.pwszVal);
+		PropVariantClear(&HardwareId);
 
-	FString TempFilterId(RenderFilterId);
-	CoTaskMemFree(RenderFilterId);
+		OutInfo.FilterId = TempFilterId;
 
-	TComPtr<IMMDevice> RenderDevnode;
-	Result = DeviceEnumerator->GetDevice(*TempFilterId, &RenderDevnode);
-	if (FAILED(Result))
-	{
-		UE_LOG(LogAudioEnumeration, Verbose, TEXT("FWindowsMMDeviceCache::EnumerateHardwareTopology failed to fetch render node: 0x%x"), Result);
-		return false;
-	}
-
-	TComPtr<IPropertyStore> PropertyStore;
-	Result = RenderDevnode->OpenPropertyStore(STGM_READ, &PropertyStore);
-	if (FAILED(Result))
-	{
-		UE_LOG(LogAudioEnumeration, Verbose, TEXT("FWindowsMMDeviceCache::EnumerateHardwareTopology failed to open property store: 0x%x"), Result);
+		return true;
+	#else
 		return false;
+	#endif //PLATFORM_WINDOWS
 	}
 
-	PROPVARIANT HardwareId;
-	PropVariantInit(&HardwareId);
-	Result = PropertyStore->GetValue(PKEY_Device_InstanceId, &HardwareId);
-	if (FAILED(Result) || !HardwareId.pwszVal)
+	void FWindowsMMDeviceCache::EnumerateEndpoints()
 	{
-		UE_LOG(LogAudioEnumeration, Verbose, TEXT("FWindowsMMDeviceCache::EnumerateHardwareTopology failed to fetch hardward Id: 0x%x"), Result);
-		return false;
-	}
-
-	OutInfo.HardwareId = FString(HardwareId.pwszVal);
-	PropVariantClear(&HardwareId);
-
-	OutInfo.FilterId = TempFilterId;
-
-	return true;
-#else
-	return false;
-#endif //PLATFORM_WINDOWS
-}
-
-void FWindowsMMDeviceCache::EnumerateEndpoints()
-{
-	SCOPED_NAMED_EVENT(FWindowsMMDeviceCache_EnumerateEndpoints, FColor::Blue);
+		SCOPED_NAMED_EVENT(FWindowsMMDeviceCache_EnumerateEndpoints, FColor::Blue);
 
-	// Build a new cache from scratch.
-	TMap<FName, FCacheEntry> NewCache;
+		// Build a new cache from scratch.
+		TMap<FName, FCacheEntry> NewCache;
 
-	// Get Device Enumerator.
-	if (DeviceEnumerator)
-	{
-		// Get Render Device Collection. (note we ask for ALL states, which include disabled/unplugged devices.).
-		TComPtr<IMMDeviceCollection> DeviceCollection;
-		uint32 DeviceCount = 0;
-		if (SUCCEEDED(DeviceEnumerator->EnumAudioEndpoints(eAll, DEVICE_STATEMASK_ALL, &DeviceCollection)) && DeviceCollection &&
-			SUCCEEDED(DeviceCollection->GetCount(&DeviceCount)))
+		// Get Device Enumerator.
+		if (DeviceEnumerator)
 		{
-			for (uint32 i = 0; i < DeviceCount; ++i)
+			// Get Render Device Collection. (note we ask for ALL states, which include disabled/unplugged devices.).
+			TComPtr<IMMDeviceCollection> DeviceCollection;
+			uint32 DeviceCount = 0;
+			if (SUCCEEDED(DeviceEnumerator->EnumAudioEndpoints(eAll, DEVICE_STATEMASK_ALL, &DeviceCollection)) && DeviceCollection &&
+				SUCCEEDED(DeviceCollection->GetCount(&DeviceCount)))
 			{
-				TComPtr<IMMDevice> Device;
-				if (SUCCEEDED(DeviceCollection->Item(i, &Device)) && Device)
+				for (uint32 i = 0; i < DeviceCount; ++i)
 				{
-					// Get the device id string (guid)
-					Audio::FScopeComString DeviceIdString;
-					if (SUCCEEDED(Device->GetId(&DeviceIdString.StringPtr)) && DeviceIdString)
+					TComPtr<IMMDevice> Device;
+					if (SUCCEEDED(DeviceCollection->Item(i, &Device)) && Device)
 					{
-						FCacheEntry Info{ DeviceIdString.Get() };
-
-						// Enumerate props into our info object.
-						EnumerateDeviceProps(Device, Info);
-
-						// Enumerate hardware topology to fetch hardware Id.
-						if (bEnumerateHardware)
+						// Get the device id string (guid)
+						Audio::FScopeComString DeviceIdString;
+						if (SUCCEEDED(Device->GetId(&DeviceIdString.StringPtr)) && DeviceIdString)
 						{
-							EnumerateHardwareTopology(Device, Info);
+							FCacheEntry Info{ DeviceIdString.Get() };
+
+							// Enumerate props into our info object.
+							EnumerateDeviceProps(Device, Info);
+
+							// Enumerate hardware topology to fetch hardware Id.
+							if (IsAggregateDeviceSupportEnabled())
+							{
+								EnumerateHardwareTopology(Device, Info);
+							}
+
+							UE_LOG(LogAudioEnumeration, Verbose, TEXT("FWindowsMMDeviceCache: %s Device '%s' ID='%s'"),
+								Info.Type == EDeviceEndpointType::Capture ? TEXT("Capture") :
+								Info.Type == EDeviceEndpointType::Render ? TEXT("Render") :
+								TEXT("UNKNOWN!"),
+								*Info.DeviceId.ToString(),
+								*Info.FriendlyName
+							);
+
+							check(!NewCache.Contains(Info.DeviceId));
+							NewCache.Emplace(Info.DeviceId, Info);
 						}
-
-						UE_LOG(LogAudioEnumeration, Verbose, TEXT("FWindowsMMDeviceCache: %s Device '%s' ID='%s'"),
-							Info.Type == FCacheEntry::EEndpointType::Capture ? TEXT("Capture") :
-							Info.Type == FCacheEntry::EEndpointType::Render ? TEXT("Render") :
-							TEXT("UNKNOWN!"),
-							*Info.DeviceId.ToString(),
-							*Info.FriendlyName
-						);
-
-						check(!NewCache.Contains(Info.DeviceId));
-						NewCache.Emplace(Info.DeviceId, Info);
 					}
 				}
 			}
-		}
 
-		// Finally, Replace cache with new one.
-		{
-			FWriteScopeLock Lock(CacheMutationLock);
-			Cache = MoveTemp(NewCache);
+			// Finally, Replace cache with new one.
+			{
+				FWriteScopeLock Lock(CacheMutationLock);
+				Cache = MoveTemp(NewCache);
+			}
 		}
 	}
-}
 
-void FWindowsMMDeviceCache::EnumerateDefaults()
-{
-	auto GetDefaultDeviceID = [this](EDataFlow InDataFlow, ERole InRole, FName& OutDeviceId) -> bool
+	void FWindowsMMDeviceCache::EnumerateDefaults()
 	{
-		// Mark default device.
-		bool bSuccess = false;
-		TComPtr<IMMDevice> DefaultDevice;
-		if (SUCCEEDED(DeviceEnumerator->GetDefaultAudioEndpoint(InDataFlow, InRole, &DefaultDevice)))
+		auto GetDefaultDeviceID = [this](EDataFlow InDataFlow, ERole InRole, FName& OutDeviceId) -> bool
 		{
-			Audio::FScopeComString DeviceIdString;
-			if (SUCCEEDED(DefaultDevice->GetId(&DeviceIdString.StringPtr)) && DeviceIdString)
+			// Mark default device.
+			bool bSuccess = false;
+			TComPtr<IMMDevice> DefaultDevice;
+			if (SUCCEEDED(DeviceEnumerator->GetDefaultAudioEndpoint(InDataFlow, InRole, &DefaultDevice)))
 			{
-				OutDeviceId = DeviceIdString.Get();
-				bSuccess = true;
+				Audio::FScopeComString DeviceIdString;
+				if (SUCCEEDED(DefaultDevice->GetId(&DeviceIdString.StringPtr)) && DeviceIdString)
+				{
+					OutDeviceId = DeviceIdString.Get();
+					bSuccess = true;
+				}
 			}
-		}
-		return bSuccess;
-	};
+			return bSuccess;
+		};
 
-	// Get defaults (render, capture).
-	FWriteScopeLock Lock(CacheMutationLock);
-	static_assert((int32)EAudioDeviceRole::COUNT == ERole_enum_count, "EAudioDeviceRole should be the same as ERole");
-	for (int32 i = 0; i < ERole_enum_count; ++i)
-	{
-		FName DeviceIdName;
-		if (GetDefaultDeviceID(eRender, static_cast<ERole>(i), DeviceIdName))
+		// Get defaults (render, capture).
+		FWriteScopeLock Lock(CacheMutationLock);
+		static_assert((int32)EAudioDeviceRole::COUNT == ERole_enum_count, "EAudioDeviceRole should be the same as ERole");
+		for (int32 i = 0; i < ERole_enum_count; ++i)
 		{
-			UE_CLOG(!DeviceIdName.IsNone(), LogAudioEnumeration, Verbose, TEXT("FWindowsMMDeviceCache: Default Render Role='%s', Device='%s'"), ToString((EAudioDeviceRole)i), *GetFriendlyName(DeviceIdName));
-			DefaultRenderId[i] = DeviceIdName;
-		}
-		if (GetDefaultDeviceID(eCapture, static_cast<ERole>(i), DeviceIdName))
-		{
-			UE_CLOG(!DeviceIdName.IsNone(), LogAudioEnumeration, Verbose, TEXT("FWindowsMMDeviceCache: Default Capture Role='%s', Device='%s'"), ToString((EAudioDeviceRole)i), *GetFriendlyName(DeviceIdName));
-			DefaultCaptureId[i] = DeviceIdName;
+			FName DeviceIdName;
+			if (GetDefaultDeviceID(eRender, static_cast<ERole>(i), DeviceIdName))
+			{
+				UE_CLOG(!DeviceIdName.IsNone(), LogAudioEnumeration, Verbose, TEXT("FWindowsMMDeviceCache: Default Render Role='%s', Device='%s'"), ToString((EAudioDeviceRole)i), *GetFriendlyName(DeviceIdName));
+				DefaultRenderId[i] = DeviceIdName;
+			}
+			if (GetDefaultDeviceID(eCapture, static_cast<ERole>(i), DeviceIdName))
+			{
+				UE_CLOG(!DeviceIdName.IsNone(), LogAudioEnumeration, Verbose, TEXT("FWindowsMMDeviceCache: Default Capture Role='%s', Device='%s'"), ToString((EAudioDeviceRole)i), *GetFriendlyName(DeviceIdName));
+				DefaultCaptureId[i] = DeviceIdName;
+			}
 		}
 	}
-}
-
-void FWindowsMMDeviceCache::OnDefaultCaptureDeviceChanged(const EAudioDeviceRole InAudioDeviceRole, const FString& DeviceId)
-{
-	FWriteScopeLock WriteLock(CacheMutationLock);
-	check(InAudioDeviceRole < EAudioDeviceRole::COUNT);
-	DefaultCaptureId[(int32)InAudioDeviceRole] = *DeviceId;
-}
 
-void FWindowsMMDeviceCache::OnDefaultRenderDeviceChanged(const EAudioDeviceRole InAudioDeviceRole, const FString& DeviceId)
-{
-	FWriteScopeLock WriteLock(CacheMutationLock);
-	check(InAudioDeviceRole < EAudioDeviceRole::COUNT);
-	DefaultRenderId[(int32)InAudioDeviceRole] = *DeviceId;
-}
+	void FWindowsMMDeviceCache::OnDefaultCaptureDeviceChanged(const EAudioDeviceRole InAudioDeviceRole, const FString& DeviceId)
+	{
+		FWriteScopeLock WriteLock(CacheMutationLock);
+		check(InAudioDeviceRole < EAudioDeviceRole::COUNT);
+		DefaultCaptureId[(int32)InAudioDeviceRole] = *DeviceId;
+	}
 
-void FWindowsMMDeviceCache::OnDeviceAdded(const FString& DeviceId, bool bIsRender)
-{
-	if (TOptional<FCacheEntry> NewDeviceEntry = BuildCacheEntry(DeviceId))
+	void FWindowsMMDeviceCache::OnDefaultRenderDeviceChanged(const EAudioDeviceRole InAudioDeviceRole, const FString& DeviceId)
 	{
 		FWriteScopeLock WriteLock(CacheMutationLock);
-		check(!Cache.Contains(NewDeviceEntry->DeviceId));
-		Cache.Emplace(NewDeviceEntry->DeviceId, MoveTemp(*NewDeviceEntry));
+		check(InAudioDeviceRole < EAudioDeviceRole::COUNT);
+		DefaultRenderId[(int32)InAudioDeviceRole] = *DeviceId;
 	}
-	else
+
+	void FWindowsMMDeviceCache::OnDeviceAdded(const FString& DeviceId, bool bIsRender)
 	{
-		UE_LOG(LogAudioEnumeration, Warning, TEXT("FWindowsMMDeviceCache::OnDeviceAdded: Failed to add DeviceID='%s' to cache. "), *DeviceId);
+		if (TOptional<FCacheEntry> NewDeviceEntry = BuildCacheEntry(DeviceId))
+		{
+			FWriteScopeLock WriteLock(CacheMutationLock);
+			check(!Cache.Contains(NewDeviceEntry->DeviceId));
+			Cache.Emplace(NewDeviceEntry->DeviceId, MoveTemp(*NewDeviceEntry));
+		}
+		else
+		{
+			UE_LOG(LogAudioEnumeration, Warning, TEXT("FWindowsMMDeviceCache::OnDeviceAdded: Failed to add DeviceID='%s' to cache. "), *DeviceId);
+		}
 	}
-}
 
-void FWindowsMMDeviceCache::OnDeviceRemoved(const FString& DeviceId, bool)
-{
-	FWriteScopeLock WriteLock(CacheMutationLock);
-	FName DeviceIdName = *DeviceId;
-	UE_CLOG(!Cache.Contains(DeviceIdName), LogAudioEnumeration, Warning, TEXT("FWindowsMMDeviceCache::OnDeviceRemoved: DeviceId='%s' was not in the cache. "), *DeviceId);
-	Cache.Remove(DeviceIdName);
-}
+	void FWindowsMMDeviceCache::OnDeviceRemoved(const FString& DeviceId, bool)
+	{
+		FWriteScopeLock WriteLock(CacheMutationLock);
+		FName DeviceIdName = *DeviceId;
+		UE_CLOG(!Cache.Contains(DeviceIdName), LogAudioEnumeration, Warning, TEXT("FWindowsMMDeviceCache::OnDeviceRemoved: DeviceId='%s' was not in the cache. "), *DeviceId);
+		Cache.Remove(DeviceIdName);
+	}
 
-TOptional<FWindowsMMDeviceCache::FCacheEntry> FWindowsMMDeviceCache::BuildCacheEntry(const FString& DeviceId)
-{
-	if (ensure(DeviceEnumerator))
+	TOptional<FWindowsMMDeviceCache::FCacheEntry> FWindowsMMDeviceCache::BuildCacheEntry(const FString& DeviceId)
 	{
-		TComPtr<IMMDevice> Device;
-		if (SUCCEEDED(DeviceEnumerator->GetDevice(*DeviceId, &Device)))
+		if (ensure(DeviceEnumerator))
 		{
-			FCacheEntry Info{ *DeviceId };
-			if (EnumerateDeviceProps(Device, Info))
+			TComPtr<IMMDevice> Device;
+			if (SUCCEEDED(DeviceEnumerator->GetDevice(*DeviceId, &Device)))
 			{
-				return Info;
+				FCacheEntry Info{ *DeviceId };
+				if (EnumerateDeviceProps(Device, Info))
+				{
+					return Info;
+				}
 			}
 		}
+		return {};
 	}
-	return {};
-}
 
-FString FWindowsMMDeviceCache::GetFriendlyName(FName InDeviceId) const
-{
-	if (const FCacheEntry* Entry = Cache.Find(InDeviceId))
+	FString FWindowsMMDeviceCache::GetFriendlyName(FName InDeviceId) const
 	{
-		return Entry->FriendlyName;
+		if (const FCacheEntry* Entry = Cache.Find(InDeviceId))
+		{
+			return Entry->FriendlyName;
+		}
+		return TEXT("Unknown");
 	}
-	return TEXT("Unknown");
-}
-
-void FWindowsMMDeviceCache::OnDeviceStateChanged(const FString& DeviceId, const EAudioDeviceState InState, bool)
-{
-	FName DeviceIdName = *DeviceId;
-	
-	// NOTE: If entry does not exist that's likely because a state change has preempted the OnDeviceAdded call.
 
-	// Scope for Read-lock on Cache Map.
-	FReadScopeLock ReadLock(CacheMutationLock);
-	if (FCacheEntry* Entry = Cache.Find(DeviceIdName))
+	void FWindowsMMDeviceCache::OnDeviceStateChanged(const FString& DeviceId, const EAudioDeviceState InState, bool)
 	{
-		// Inner Write-Lock on Entry.
-		FWriteScopeLock WriteLock(Entry->MutationLock);
+		FName DeviceIdName = *DeviceId;
+		
+		// NOTE: If entry does not exist that's likely because a state change has preempted the OnDeviceAdded call.
 
-		UE_CLOG(WindowsMMCvarUtils::ShouldLogDeviceSwaps(), LogAudioEnumeration, Verbose, TEXT("FWindowsMMDeviceCache: DeviceName='%s' - DeviceID='%s' state changed from '%s' to '%s'."),
-			*Entry->FriendlyName, *DeviceId, ToString(Entry->State), ToString(InState));
+		// Scope for Read-lock on Cache Map.
+		FReadScopeLock ReadLock(CacheMutationLock);
+		if (FCacheEntry* Entry = Cache.Find(DeviceIdName))
+		{
+			// Inner Write-Lock on Entry.
+			FWriteScopeLock WriteLock(Entry->MutationLock);
 
-		Entry->State = InState;
-	}
-}
+			UE_CLOG(WindowsMMCvarUtils::ShouldLogDeviceSwaps(), LogAudioEnumeration, Verbose, TEXT("FWindowsMMDeviceCache: DeviceName='%s' - DeviceID='%s' state changed from '%s' to '%s'."),
+				*Entry->FriendlyName, *DeviceId, ToString(Entry->State), ToString(InState));
 
-void FWindowsMMDeviceCache::OnFormatChanged(const FString& InDeviceId, const FFormatChangedData& InFormat)
-{
-	FName DeviceName(InDeviceId);
-	bool bNeedToEnumerateChannels = false;
-	bool bDirty = false;
+			Entry->State = InState;
+		}
+	}
 
-	FReadScopeLock MapReadLock(CacheMutationLock);
-	if (FCacheEntry* Found = Cache.Find(DeviceName))
+	void FWindowsMMDeviceCache::OnFormatChanged(const FString& InDeviceId, const FFormatChangedData& InFormat)
 	{
-		// Make a copy of the entry
-		FCacheEntry EntryCopy(InDeviceId);
-		{
-			FReadScopeLock FoundReadLock(Found->MutationLock);
-			EntryCopy = *Found;
-		}
+		FName DeviceName(InDeviceId);
+		bool bNeedToEnumerateChannels = false;
+		bool bDirty = false;
 
-		if (EntryCopy.NumChannels != InFormat.NumChannels)
+		FReadScopeLock MapReadLock(CacheMutationLock);
+		if (FCacheEntry* Found = Cache.Find(DeviceName))
 		{
-			UE_CLOG(WindowsMMCvarUtils::ShouldLogDeviceSwaps(), LogAudioEnumeration, Verbose, TEXT("FWindowsMMDeviceCache: DeviceID='%s', Name='%s' changed default format from %d channels to %d."), *InDeviceId, *EntryCopy.FriendlyName, EntryCopy.NumChannels, InFormat.NumChannels);
-			EntryCopy.NumChannels = InFormat.NumChannels;
-			bNeedToEnumerateChannels = true;
-			bDirty = true;
+			// Make a copy of the entry
+			FCacheEntry EntryCopy(InDeviceId);
+			{
+				FReadScopeLock FoundReadLock(Found->MutationLock);
+				EntryCopy = *Found;
+			}
+
+			if (EntryCopy.NumChannels != InFormat.NumChannels)
+			{
+				UE_CLOG(WindowsMMCvarUtils::ShouldLogDeviceSwaps(), LogAudioEnumeration, Verbose, TEXT("FWindowsMMDeviceCache: DeviceID='%s', Name='%s' changed default format from %d channels to %d."), *InDeviceId, *EntryCopy.FriendlyName, EntryCopy.NumChannels, InFormat.NumChannels);
+				EntryCopy.NumChannels = InFormat.NumChannels;
+				bNeedToEnumerateChannels = true;
+				bDirty = true;
+			}
+			if (EntryCopy.SampleRate != InFormat.SampleRate)
+			{
+				UE_CLOG(WindowsMMCvarUtils::ShouldLogDeviceSwaps(), LogAudioEnumeration, Verbose, TEXT("FWindowsMMDeviceCache: DeviceID='%s', Name='%s' changed default format from %dhz to %dhz."), *InDeviceId, *EntryCopy.FriendlyName, EntryCopy.SampleRate, InFormat.SampleRate);
+				EntryCopy.SampleRate = InFormat.SampleRate;
+				bDirty = true;
+			}
+			if (EntryCopy.ChannelBitmask != InFormat.ChannelBitmask)
+			{
+				UE_CLOG(WindowsMMCvarUtils::ShouldLogDeviceSwaps(), LogAudioEnumeration, Verbose, TEXT("FWindowsMMDeviceCache: DeviceID='%s', Name='%s' changed default format from 0x%x to 0x%x bitmask"), *InDeviceId, *EntryCopy.FriendlyName, EntryCopy.ChannelBitmask, InFormat.ChannelBitmask);
+				EntryCopy.ChannelBitmask = InFormat.ChannelBitmask;
+				bNeedToEnumerateChannels = true;
+				bDirty = true;
+			}
+
+			if (bNeedToEnumerateChannels)
+			{
+				UE_CLOG(WindowsMMCvarUtils::ShouldLogDeviceSwaps(), LogAudioEnumeration, Verbose, TEXT("FWindowsMMDeviceCache: Channel Change, DeviceID='%s', Name='%s' OLD=[%s]"), *InDeviceId, *EntryCopy.FriendlyName, *ToFString(EntryCopy.OutputChannels));
+				EnumerateChannelMask(InFormat.ChannelBitmask, EntryCopy);
+				UE_CLOG(WindowsMMCvarUtils::ShouldLogDeviceSwaps(), LogAudioEnumeration, Verbose, TEXT("FWindowsMMDeviceCache: Channel Change, DeviceID='%s', Name='%s' NEW=[%s]"), *InDeviceId, *EntryCopy.FriendlyName, *ToFString(EntryCopy.OutputChannels));
+			}
+
+			// Update the entire entry with one write.
+			if (bDirty)
+			{
+				FWriteScopeLock FoundWriteLock(Found->MutationLock);
+				*Found = EntryCopy;
+			}
 		}
-		if (EntryCopy.SampleRate != InFormat.SampleRate)
+	}
+
+	void FWindowsMMDeviceCache::MakeDeviceInfo(const FCacheEntry& InEntry, FName InDefaultDevice, FAudioPlatformDeviceInfo& OutInfo) const
+	{
+		OutInfo.Reset();
+		OutInfo.Name = InEntry.FriendlyName;
+		OutInfo.DeviceId = InEntry.DeviceId.GetPlainNameString();
+		OutInfo.NumChannels = InEntry.NumChannels;
+		OutInfo.SampleRate = InEntry.SampleRate;
+		OutInfo.OutputChannelArray = InEntry.OutputChannels;
+		OutInfo.Format = EAudioMixerStreamDataFormat::Float;
+		OutInfo.bIsSystemDefault = InEntry.DeviceId == InDefaultDevice;
+	}
+
+	TArray<FAudioPlatformDeviceInfo> FWindowsMMDeviceCache::GetAllActiveOutputDevices() const
+	{
+		SCOPED_NAMED_EVENT(FWindowsMMDeviceCache_GetAllActiveOutputDevices, FColor::Blue);
+
+		// Find all active devices.
+		TArray<FAudioPlatformDeviceInfo> ActiveDevices;
+
+		if (IsAggregateDeviceSupportEnabled())
 		{
-			UE_CLOG(WindowsMMCvarUtils::ShouldLogDeviceSwaps(), LogAudioEnumeration, Verbose, TEXT("FWindowsMMDeviceCache: DeviceID='%s', Name='%s' changed default format from %dhz to %dhz."), *InDeviceId, *EntryCopy.FriendlyName, EntryCopy.SampleRate, InFormat.SampleRate);
-			EntryCopy.SampleRate = InFormat.SampleRate;
-			bDirty = true;
+			// Determine if there are any aggregate devices and place them at the beginning of the array
+			ActiveDevices = SynthesizeAggregateDeviceList(EDeviceEndpointType::Render);
+			Algo::Sort(ActiveDevices, [](const FAudioPlatformDeviceInfo& ItemA, const FAudioPlatformDeviceInfo& ItemB)
+			{
+				return ItemA.Name.Compare(ItemB.Name) < 0;
+			});
 		}
-		if (EntryCopy.ChannelBitmask != InFormat.ChannelBitmask)
+
+		TArray<FAudioPlatformDeviceInfo> NonAggregateDevices;
+
+		// Read lock
+		FReadScopeLock ReadLock(CacheMutationLock);
+		NonAggregateDevices.Reserve(Cache.Num());
+
+		// Ask for defaults once, as we are inside a read lock.
+		FName DefaultRenderDeviceId = GetDefaultOutputDevice_NoLock();
+
+		// Walk cache, read lock for each entry.
+		for (const auto& i : Cache)
 		{
-			UE_CLOG(WindowsMMCvarUtils::ShouldLogDeviceSwaps(), LogAudioEnumeration, Verbose, TEXT("FWindowsMMDeviceCache: DeviceID='%s', Name='%s' changed default format from 0x%x to 0x%x bitmask"), *InDeviceId, *EntryCopy.FriendlyName, EntryCopy.ChannelBitmask, InFormat.ChannelBitmask);
-			EntryCopy.ChannelBitmask = InFormat.ChannelBitmask;
-			bNeedToEnumerateChannels = true;
-			bDirty = true;
+			// Read lock on each entry.
+			FReadScopeLock CacheEntryReadLock(i.Value.MutationLock);
+			if (i.Value.State == EAudioDeviceState::Active &&
+				i.Value.Type == EDeviceEndpointType::Render)
+			{
+				FAudioPlatformDeviceInfo& Info = NonAggregateDevices.Emplace_GetRef();
+				MakeDeviceInfo(i.Value, DefaultRenderDeviceId, Info);
+			}
 		}
 
-		if (bNeedToEnumerateChannels)
+		Algo::Sort(NonAggregateDevices, [](const FAudioPlatformDeviceInfo& ItemA, const FAudioPlatformDeviceInfo& ItemB)
 		{
-			UE_CLOG(WindowsMMCvarUtils::ShouldLogDeviceSwaps(), LogAudioEnumeration, Verbose, TEXT("FWindowsMMDeviceCache: Channel Change, DeviceID='%s', Name='%s' OLD=[%s]"), *InDeviceId, *EntryCopy.FriendlyName, *ToFString(EntryCopy.OutputChannels));
-			EnumerateChannelMask(InFormat.ChannelBitmask, EntryCopy);
-			UE_CLOG(WindowsMMCvarUtils::ShouldLogDeviceSwaps(), LogAudioEnumeration, Verbose, TEXT("FWindowsMMDeviceCache: Channel Change, DeviceID='%s', Name='%s' NEW=[%s]"), *InDeviceId, *EntryCopy.FriendlyName, *ToFString(EntryCopy.OutputChannels));
-		}
+			// See if both devices belong to an aggregate group. If so, then we use channel numbers
+			// as the secondary sort criteria.
+			FString ItemAHardwareName = ExtractAggregateDeviceName(ItemA.Name);
+			FString ItemBHardwareName = ExtractAggregateDeviceName(ItemB.Name);
+			
+			int32 Result = ItemAHardwareName.Compare(ItemBHardwareName);
+			if (Result == 0)
+			{
+				int32 ItemAChannelNumber = ExtractAggregateChannelNumber(ItemA.Name);
+				int32 ItemBChannelNumber = ExtractAggregateChannelNumber(ItemB.Name);
+				return ItemAChannelNumber < ItemBChannelNumber;
+			}
+
+			return Result < 0;
+		});
 
-		// Update the entire entry with one write.
-		if (bDirty)
+		ActiveDevices.Append(NonAggregateDevices);
+		
+		// RVO
+		return ActiveDevices;
+	}
+
+	bool FWindowsMMDeviceCache::IsAggregateHardwareDeviceId(const FName InDeviceID) const
+	{
+		FReadScopeLock ReadLock(CacheMutationLock);
+
+		TArray<FCacheEntry> ActiveDevices;
+		for (const auto& Entry : Cache)
 		{
-			FWriteScopeLock FoundWriteLock(Found->MutationLock);
-			*Found = EntryCopy;
+			// HardwareId is used for DeviceId with aggregate devices
+			if (Entry.Value.HardwareId == InDeviceID)
+			{
+				return true;
+			}
 		}
-	}
-}
 
-void FWindowsMMDeviceCache::MakeDeviceInfo(const FCacheEntry& InEntry, FName InDefaultDevice, FAudioPlatformDeviceInfo& OutInfo) const
-{
-	OutInfo.Reset();
-	OutInfo.Name = InEntry.FriendlyName;
-	OutInfo.DeviceId = InEntry.DeviceId.GetPlainNameString();
-	OutInfo.NumChannels = InEntry.NumChannels;
-	OutInfo.SampleRate = InEntry.SampleRate;
-	OutInfo.OutputChannelArray = InEntry.OutputChannels;
-	OutInfo.Format = EAudioMixerStreamDataFormat::Float;
-	OutInfo.bIsSystemDefault = InEntry.DeviceId == InDefaultDevice;
-}
-
-TArray<Audio::FAudioPlatformDeviceInfo> FWindowsMMDeviceCache::GetAllActiveOutputDevices() const
-{
-	SCOPED_NAMED_EVENT(FWindowsMMDeviceCache_GetAllActiveOutputDevices, FColor::Blue);
+		return false;
+	}
 
-	// Find all active devices.
-	TArray<FAudioPlatformDeviceInfo> ActiveDevices;
+	TOptional<FAudioPlatformDeviceInfo> FWindowsMMDeviceCache::GetAggregateHardwareDeviceInfo(const FName InHardwareId, const EDeviceEndpointType InEndpointType) const
+	{
+		TSet<FCacheEntry, FCacheKeyFuncs> UniqueHardwareIds;
+		TMap<FName, FDeviceChannelInfo> DeviceChannelInfos;
 
-	// Read lock
-	FReadScopeLock ReadLock(CacheMutationLock);
-	ActiveDevices.Reserve(Cache.Num());
+		GetHardwareInfo(UniqueHardwareIds, DeviceChannelInfos, InEndpointType);
 
-	// Ask for defaults once, as we are inside a read lock.
-	FName DefaultRenderDeviceId = GetDefaultOutputDevice_NoLock();
+		for (const auto& Entry : UniqueHardwareIds)
+		{
+			if (Entry.HardwareId == InHardwareId)
+			{
+				if (ensure(DeviceChannelInfos.Contains(Entry.HardwareId)))
+				{
+					return CreateAggregateDeviceInfo(Entry, DeviceChannelInfos[Entry.HardwareId]);
+				}
+			}
+		}
 
-	// Walk cache, read lock for each entry.
-	for (const auto& i : Cache)
+		return {};
+	}
+	
+	TArray<FAudioPlatformDeviceInfo> FWindowsMMDeviceCache::GetLogicalAggregateDevices(const FName InHardwareId, const EDeviceEndpointType InEndpointType) const
 	{
-		// Read lock on each entry.
-		FReadScopeLock CacheEntryReadLock(i.Value.MutationLock);
-		if (i.Value.State == EAudioDeviceState::Active &&
-			i.Value.Type == FCacheEntry::EEndpointType::Render)
+		TArray<FAudioPlatformDeviceInfo> AggregateDevice;
+		FReadScopeLock ReadLock(CacheMutationLock);
+
+		TArray<FCacheEntry> ActiveDevices;
+		for (const auto& Entry : Cache)
 		{
-			FAudioPlatformDeviceInfo& Info = ActiveDevices.Emplace_GetRef();
-			MakeDeviceInfo(i.Value, DefaultRenderDeviceId, Info);
+			if (Entry.Value.HardwareId == InHardwareId &&
+				Entry.Value.State == EAudioDeviceState::Active &&
+				Entry.Value.Type == InEndpointType)
+			{
+				ActiveDevices.Emplace(Entry.Value);
+			}
 		}
-	}
 
-	// RVO
-	return ActiveDevices;
-}
+		// Sort the devices according to their filter id which usually ends with a vendor specific
+		// string which is sortable
+		Algo::Sort(ActiveDevices, [](const FCacheEntry& InEntryA, const FCacheEntry& InEntryB)
+		{
+			return InEntryA.FilterId.Compare(InEntryB.FilterId) < 0;
+		});
 
-TArray<Audio::FAudioPlatformDeviceInfo> FWindowsMMDeviceCache::GetAggregateDeviceInfo(FName InDeviceId) const
-{
-	TArray<FAudioPlatformDeviceInfo> AggregateDevice;
-	FReadScopeLock ReadLock(CacheMutationLock);
+		for (const FCacheEntry& Device : ActiveDevices)
+		{
+			FAudioPlatformDeviceInfo& Info = AggregateDevice.Emplace_GetRef();
+			MakeDeviceInfo(Device, GetDefaultOutputDevice_NoLock(), Info);
+		}
 
-	if (const FCacheEntry* DeviceEntry = Cache.Find(InDeviceId))
+		return AggregateDevice;
+	}
+	
+	void FWindowsMMDeviceCache::GetHardwareInfo(TSet<FCacheEntry, FCacheKeyFuncs>& OutUniqueHardwareIds,
+												TMap<FName, FDeviceChannelInfo>& OutDeviceChannelInfos,
+												EDeviceEndpointType InType) const
 	{
-		if (!DeviceEntry->HardwareId.IsEmpty())
+		FReadScopeLock ReadLock(CacheMutationLock);
+		for (const auto& Entry : Cache)
 		{
-			TArray<FCacheEntry> ActiveDevices;
-
-			for (const auto& Entry : Cache)
+			if (Entry.Value.State == EAudioDeviceState::Active && Entry.Value.Type == InType)
 			{
-				if (Entry.Value.HardwareId == DeviceEntry->HardwareId &&
-					Entry.Value.State == EAudioDeviceState::Active &&
-					Entry.Value.Type == DeviceEntry->Type)
+				OutUniqueHardwareIds.Add(Entry.Value);
+
+				// Accumulate channel counts for the aggregate devices
+				if (OutDeviceChannelInfos.Contains(Entry.Value.HardwareId))
 				{
-					ActiveDevices.Emplace(Entry.Value);
+					FDeviceChannelInfo* ChannelInfo = OutDeviceChannelInfos.Find(Entry.Value.HardwareId);
+					ensure(ChannelInfo->LogicDeviceChannelCount == Entry.Value.NumChannels);
+
+					ChannelInfo->TotalChannelCount += Entry.Value.NumChannels;
+				}
+				else
+				{
+					FDeviceChannelInfo& ChannelInfo = OutDeviceChannelInfos.Add(Entry.Value.HardwareId);
+					ChannelInfo.LogicDeviceChannelCount = Entry.Value.NumChannels;
+					ChannelInfo.TotalChannelCount += Entry.Value.NumChannels;
 				}
 			}
+		}
+	}
 
-			// Sort the devices according to their filter id which usually ends with a vendor specific
-			// string which is sortable
-			Algo::Sort(ActiveDevices, [](const FCacheEntry& InEntryA, const FCacheEntry& InEntryB)
-			{
-				return InEntryA.FilterId.Compare(InEntryB.FilterId) < 0;
-			});
+	FAudioPlatformDeviceInfo FWindowsMMDeviceCache::CreateAggregateDeviceInfo(const FCacheEntry& InCacheEntry, const FDeviceChannelInfo& InDeviceChannelInfo)
+	{
+		FAudioPlatformDeviceInfo Info;
+		
+		// Every Windows audio device contains the hardware name at the end of the string in parentheses
+		Info.Name = ExtractAggregateDeviceName(InCacheEntry.FriendlyName);
+		// Synthesize a device id from the hardware id which is unique
+		Info.DeviceId = InCacheEntry.HardwareId.GetPlainNameString();
+		Info.NumChannels = InCacheEntry.NumChannels;
+		Info.SampleRate = InCacheEntry.SampleRate;
+		Info.OutputChannelArray = InCacheEntry.OutputChannels;
+		Info.Format = EAudioMixerStreamDataFormat::Float;
+		Info.bIsSystemDefault = false;
+
+		int32 NumDirectOuts = InDeviceChannelInfo.TotalChannelCount - InDeviceChannelInfo.LogicDeviceChannelCount;
+		Info.NumDirectOutChannels = FMath::Max(0, NumDirectOuts);
+
+		return Info;
+	}
+	
+	TArray<FAudioPlatformDeviceInfo> FWindowsMMDeviceCache::SynthesizeAggregateDeviceList(EDeviceEndpointType InType) const
+	{
+		TSet<FCacheEntry, FCacheKeyFuncs> UniqueHardwareIds;
+		TMap<FName, FDeviceChannelInfo> DeviceChannelInfos;
 
-			for (const FCacheEntry& Device : ActiveDevices)
+		GetHardwareInfo(UniqueHardwareIds, DeviceChannelInfos, InType);
+		TArray<FAudioPlatformDeviceInfo> AggregateDeviceList;
+
+		for (const auto& Entry : UniqueHardwareIds)
+		{
+			if (ensure(DeviceChannelInfos.Contains(Entry.HardwareId)))
 			{
-				FAudioPlatformDeviceInfo& Info = AggregateDevice.Emplace_GetRef();
-				MakeDeviceInfo(Device, GetDefaultOutputDevice_NoLock(), Info);
+				FAudioPlatformDeviceInfo Info = CreateAggregateDeviceInfo(Entry, DeviceChannelInfos[Entry.HardwareId]);
+
+				// Only consider aggregate devices that have direct out channels
+				if (Info.NumDirectOutChannels > 0)
+				{
+					AggregateDeviceList.Emplace(MoveTemp(Info));
+				}
 			}
 		}
+
+		return AggregateDeviceList;
 	}
 
-	return AggregateDevice;
-}
+	FString FWindowsMMDeviceCache::ExtractAggregateDeviceName(const FString& InName)
+	{
+		// The hardware name is in parentheses at the end of the string
+		const FRegexPattern RegexPattern(TEXT("\\(([^\\(\\)]+)\\)$"));
+		FRegexMatcher RegexMatcher(RegexPattern, InName);
 
-FName FWindowsMMDeviceCache::GetDefaultOutputDevice_NoLock() const
-{
-	if (!DefaultRenderId[(int32)EAudioDeviceRole::Console].IsNone())
+		if (RegexMatcher.FindNext())
+		{
+			return RegexMatcher.GetCaptureGroup(1);
+		}
+		
+		return InName;
+	}
+
+	int32 FWindowsMMDeviceCache::ExtractAggregateChannelNumber(const FString& InName)
 	{
-		return DefaultRenderId[(int32)EAudioDeviceRole::Console];
+		// The hardware name is in parentheses at the end of the string
+		const FRegexPattern RegexPattern(TEXT("^.+?\\(([\\d]+)-[\\d]+\\)"));
+		FRegexMatcher RegexMatcher(RegexPattern, InName);
+
+		if (RegexMatcher.FindNext())
+		{
+			FString ChannelNumber = RegexMatcher.GetCaptureGroup(1);
+			return FCString::Atoi(*ChannelNumber);
+		}
+		
+		return 0;
 	}
-	if (!DefaultRenderId[(int32)EAudioDeviceRole::Multimedia].IsNone())
+
+	FName FWindowsMMDeviceCache::GetDefaultOutputDevice_NoLock() const
 	{
-		return DefaultRenderId[(int32)EAudioDeviceRole::Multimedia];
+		if (!DefaultRenderId[(int32)EAudioDeviceRole::Console].IsNone())
+		{
+			return DefaultRenderId[(int32)EAudioDeviceRole::Console];
+		}
+		if (!DefaultRenderId[(int32)EAudioDeviceRole::Multimedia].IsNone())
+		{
+			return DefaultRenderId[(int32)EAudioDeviceRole::Multimedia];
+		}
+		return NAME_None;
 	}
-	return NAME_None;
-}
 
-TOptional<Audio::FAudioPlatformDeviceInfo> FWindowsMMDeviceCache::FindDefaultOutputDevice() const
-{
-	return FindActiveOutputDevice(NAME_None);
-}
+	TOptional<FAudioPlatformDeviceInfo> FWindowsMMDeviceCache::FindDefaultOutputDevice() const
+	{
+		return FindActiveOutputDevice(NAME_None);
+	}
 
-TOptional<Audio::FAudioPlatformDeviceInfo> FWindowsMMDeviceCache::FindActiveOutputDevice(FName InDeviceID) const
-{
-	SCOPED_NAMED_EVENT(FWindowsMMDeviceCache_FindActiveOutputDevice, FColor::Blue);
+	TOptional<FAudioPlatformDeviceInfo> FWindowsMMDeviceCache::FindActiveOutputDevice(FName InDeviceID) const
+	{
+		SCOPED_NAMED_EVENT(FWindowsMMDeviceCache_FindActiveOutputDevice, FColor::Blue);
 
-	FReadScopeLock MapReadLock(CacheMutationLock);
+		FReadScopeLock MapReadLock(CacheMutationLock);
 
-	// Ask for default here as we are inside the read lock.
-	const FName DefaultOutputDevice = GetDefaultOutputDevice_NoLock();
+		// Ask for default here as we are inside the read lock.
+		const FName DefaultOutputDevice = GetDefaultOutputDevice_NoLock();
 
-	// Asking for Default?
-	if (InDeviceID.IsNone())
-	{
-		InDeviceID = DefaultOutputDevice;
+		// Asking for Default?
 		if (InDeviceID.IsNone())
 		{
-			// No default set, fail.
-			return {};
+			InDeviceID = DefaultOutputDevice;
+			if (InDeviceID.IsNone())
+			{
+				// No default set, fail.
+				return {};
+			}
 		}
-	}
 
-	// Find entry matching that device ID.
-	if (const FCacheEntry* Found = Cache.Find(InDeviceID))
-	{
-		FReadScopeLock EntryReadLock(Found->MutationLock);
-		if (Found->State == EAudioDeviceState::Active &&
-			Found->Type == FCacheEntry::EEndpointType::Render)
+		// Find entry matching that device ID.
+		if (const FCacheEntry* Found = Cache.Find(InDeviceID))
+		{
+			FReadScopeLock EntryReadLock(Found->MutationLock);
+			if (Found->State == EAudioDeviceState::Active &&
+				Found->Type == EDeviceEndpointType::Render)
+			{
+				FAudioPlatformDeviceInfo Info;
+				MakeDeviceInfo(*Found, DefaultOutputDevice, Info);
+				return Info;
+			}
+		}
+
+		if (bIsAggregateDeviceSupportEnabled)
 		{
-			FAudioPlatformDeviceInfo Info;
-			MakeDeviceInfo(*Found, DefaultOutputDevice, Info);
-			return Info;
+			/** Returns the device info for an aggregate audio device. Note that this is a virtual device in that it is not
+			 *  a device returned by the OS device enumerator. It is a device synthesized for the purpose of identifying
+			 *  an aggregate device as a single, unified device. This device cannot be instantiated as-is. It's device Id
+			 *  can be used with GetAggregateHardwareDeviceInfo() to get an array of the logical devices that make up the 
+			 *  aggregate and can be instantiated as a group to form an aggregate device at runtime.
+			 */
+			TOptional<FAudioPlatformDeviceInfo> Info = GetAggregateHardwareDeviceInfo(InDeviceID, EDeviceEndpointType::Render);
+			if (Info.IsSet())
+			{
+				return Info;
+			}
 		}
+		
+		// Fail.
+		return {};
 	}
-	// Fail.
-	return {};
-}
 
+	bool FWindowsMMDeviceCache::IsAggregateDeviceSupportEnabled() const
+	{
+		return bIsAggregateDeviceSupportEnabled || WindowsMMCvarUtils::IsAggregateDeviceSupportCVarEnabled();
+	}
 } // namespace Audio
diff --git a/Engine/Source/Runtime/AudioDeviceEnumeration/Windows/WindowsMMDeviceEnumeration/Public/WindowsMMDeviceInfoCache.h b/Engine/Source/Runtime/AudioDeviceEnumeration/Windows/WindowsMMDeviceEnumeration/Public/WindowsMMDeviceInfoCache.h
--- a/Engine/Source/Runtime/AudioDeviceEnumeration/Windows/WindowsMMDeviceEnumeration/Public/WindowsMMDeviceInfoCache.h
+++ b/Engine/Source/Runtime/AudioDeviceEnumeration/Windows/WindowsMMDeviceEnumeration/Public/WindowsMMDeviceInfoCache.h
@@ -26,17 +26,18 @@ struct WINDOWSMMDEVICEENUMERATION_API FWindowsMMDeviceCache : IAudioMixerDeviceC
 {
 	struct FCacheEntry
 	{
+		UE_DEPRECATED(5.6, "EEndpointType is deprecated. Please use EDeviceEndpointType instead.");
 		enum class EEndpointType { Unknown, Render, Capture };
 
 		FName DeviceId;							// Key
 		FString FriendlyName;
 		FString DeviceFriendlyName;
-		EAudioDeviceState State;
+		EAudioDeviceState State = EAudioDeviceState::NotPresent;
 		int32 NumChannels = 0;
 		int32 SampleRate = 0;
-		EEndpointType Type = EEndpointType::Unknown;
+		EDeviceEndpointType Type = EDeviceEndpointType::Unknown;
 		uint32 ChannelBitmask = 0;				// Bitfield used to build output channels, for easy comparison.
-		FString HardwareId;						// Unique string of the physical hardware device this MMDevice belongs to
+		FName HardwareId;						// Unique string of the physical hardware device this MMDevice belongs to
 		FString FilterId;						// Unique identifier for this device containing product id (pid), vendor id (vid), etc.
 
 		TArray<EAudioMixerChannel::Type> OutputChannels;	// TODO. Generate this from the ChannelNum and bitmask when we are asked for it.
@@ -59,16 +60,16 @@ struct WINDOWSMMDEVICEENUMERATION_API FWindowsMMDeviceCache : IAudioMixerDeviceC
 	TMap<FName, FCacheEntry> Cache;								// DeviceID GUID -> Info.
 	FName DefaultCaptureId[(int32)EAudioDeviceRole::COUNT];		// Role -> DeviceID GUID
 	FName DefaultRenderId[(int32)EAudioDeviceRole::COUNT];		// Role -> DeviceID GUID
-	bool bEnumerateHardware = false;							// Hardware enumeration can be slow. Enumerate only if requested.
 
-	explicit FWindowsMMDeviceCache(bool bInEnumerateHardware = false);
+	FWindowsMMDeviceCache();
+	explicit FWindowsMMDeviceCache(bool bInEnableAggregateDeviceSupport);
 	virtual ~FWindowsMMDeviceCache() = default;
 
 	bool EnumerateChannelMask(uint32 InMask, FCacheEntry& OutInfo);
 
 	bool EnumerateChannelFormat(const WAVEFORMATEX* InFormat, FCacheEntry& OutInfo);
 
-	FCacheEntry::EEndpointType QueryDeviceDataFlow(const TComPtr<IMMDevice>& InDevice) const;
+	EDeviceEndpointType QueryDeviceDataFlow(const TComPtr<IMMDevice>& InDevice) const;
 
 	bool EnumerateDeviceProps(const TComPtr<IMMDevice>& InDevice, FCacheEntry& OutInfo);
 
@@ -98,13 +99,56 @@ struct WINDOWSMMDEVICEENUMERATION_API FWindowsMMDeviceCache : IAudioMixerDeviceC
 
 	virtual TArray<FAudioPlatformDeviceInfo> GetAllActiveOutputDevices() const override;
 
-	virtual TArray<FAudioPlatformDeviceInfo> GetAggregateDeviceInfo(FName InDeviceId) const override;
+	virtual bool IsAggregateHardwareDeviceId(const FName InDeviceID) const override;
 
+	virtual TOptional<FAudioPlatformDeviceInfo> GetAggregateHardwareDeviceInfo(const FName InHardwareId, const EDeviceEndpointType InEndpointType) const;
+
+	virtual TArray<FAudioPlatformDeviceInfo> GetLogicalAggregateDevices(const FName InHardwareId, const EDeviceEndpointType InEndpointType) const override;
+
+	virtual TArray<FAudioPlatformDeviceInfo> SynthesizeAggregateDeviceList(const EDeviceEndpointType InType) const;
+	
 	FName GetDefaultOutputDevice_NoLock() const;
 
 	TOptional<FAudioPlatformDeviceInfo> FindDefaultOutputDevice() const override;
 
 	TOptional<FAudioPlatformDeviceInfo> FindActiveOutputDevice(FName InDeviceID) const override;
+
+	bool IsAggregateDeviceSupportEnabled() const;
+
+	static FString ExtractAggregateDeviceName(const FString& InName);
+
+private:
+	struct FCacheKeyFuncs : BaseKeyFuncs<FWindowsMMDeviceCache::FCacheEntry, FName>
+	{
+		static KeyInitType GetSetKey(ElementInitType Element)
+		{
+			return Element.HardwareId;
+		}
+
+		static bool Matches(KeyInitType A, KeyInitType B)
+		{
+			return A.IsEqual(B, ENameCase::CaseSensitive);
+		}
+
+		static uint32 GetKeyHash(KeyInitType Key)
+		{
+			return GetTypeHash(Key);
+		}
+	};
+
+	struct FDeviceChannelInfo
+	{
+		uint32 LogicDeviceChannelCount = 0;
+		uint32 TotalChannelCount = 0;
+	};
+
+	bool bIsAggregateDeviceSupportEnabled = false;
+	
+	void GetHardwareInfo(TSet<FCacheEntry, FCacheKeyFuncs>& OutUniqueHardwareIds, TMap<FName, FDeviceChannelInfo>& OutDeviceChannelInfos, EDeviceEndpointType InType) const;
+
+	static FAudioPlatformDeviceInfo CreateAggregateDeviceInfo(const FCacheEntry& InCacheEntry, const FDeviceChannelInfo& InDeviceChannelInfo);
+	
+	static int32 ExtractAggregateChannelNumber(const FString& InName);
 };
 
 }// namespace Audio
diff --git a/Engine/Source/Runtime/AudioMixer/Private/AudioMixerDevice.cpp b/Engine/Source/Runtime/AudioMixer/Private/AudioMixerDevice.cpp
--- a/Engine/Source/Runtime/AudioMixer/Private/AudioMixerDevice.cpp
+++ b/Engine/Source/Runtime/AudioMixer/Private/AudioMixerDevice.cpp
@@ -36,6 +36,7 @@
 
 #if WITH_EDITOR
 #include "AudioEditorModule.h"
+#include "AudioEditorSettings.h"
 #endif // WITH_EDITOR
 
 #ifndef CASE_ENUM_TO_TEXT
@@ -553,6 +554,17 @@ namespace Audio
 			OpenStreamParams.SampleRate = SampleRate;
 			OpenStreamParams.AudioMixer = this;
 			OpenStreamParams.MaxSources = GetMaxSources();
+#if WITH_EDITOR
+			{
+				IAudioEditorModule& AudioEditorModule = FModuleManager::LoadModuleChecked<IAudioEditorModule>("AudioEditor");
+				TOptional<FAudioEditorDeviceSettings> DeviceSettings = AudioEditorModule.GetAudioEditorDeviceSettings();
+				if (DeviceSettings.IsSet())
+				{
+					OpenStreamParams.bUseSystemAudioDevice = DeviceSettings->bUseSystemDevice;
+					OpenStreamParams.AudioDeviceId = DeviceSettings->DeviceId;
+				}
+			}
+#endif //WITH_EDITOR
 
 			FString DefaultDeviceName = AudioMixerPlatform->GetDefaultDeviceName();
 
@@ -562,7 +574,13 @@ namespace Audio
 				DefaultDeviceName = IHeadMountedDisplayModule::Get().GetAudioOutputDevice();
 			}
 
-			if (!DefaultDeviceName.IsEmpty())
+			FString DeviceToOpen = DefaultDeviceName;
+			if (!OpenStreamParams.bUseSystemAudioDevice)
+			{
+				DeviceToOpen = OpenStreamParams.AudioDeviceId;
+			}
+			
+			if (!DeviceToOpen.IsEmpty())
 			{
 				uint32 NumOutputDevices = 0;
 				AudioMixerPlatform->GetNumOutputDevices(NumOutputDevices);
@@ -572,7 +590,7 @@ namespace Audio
 					FAudioPlatformDeviceInfo DeviceInfo;
 					AudioMixerPlatform->GetOutputDeviceInfo(i, DeviceInfo);
 
-					if (DeviceInfo.Name == DefaultDeviceName || DeviceInfo.DeviceId == DefaultDeviceName)
+					if (DeviceInfo.Name == DeviceToOpen || DeviceInfo.DeviceId == DeviceToOpen)
 					{
 						OpenStreamParams.OutputDeviceIndex = i;
 
diff --git a/Engine/Source/Runtime/AudioMixerCore/Public/AudioMixer.h b/Engine/Source/Runtime/AudioMixerCore/Public/AudioMixer.h
--- a/Engine/Source/Runtime/AudioMixerCore/Public/AudioMixer.h
+++ b/Engine/Source/Runtime/AudioMixerCore/Public/AudioMixer.h
@@ -223,6 +223,8 @@ namespace Audio
 		bool bIsMainAudioMixer;
 	};
 
+	enum class EDeviceEndpointType { Unknown, Render, Capture };
+
 	// Interface for Caching Device Info.
 	class IAudioPlatformDeviceInfoCache
 	{
@@ -233,7 +235,12 @@ namespace Audio
 		virtual TOptional<FAudioPlatformDeviceInfo> FindActiveOutputDevice(FName InDeviceID) const = 0;
 		virtual TArray<FAudioPlatformDeviceInfo> GetAllActiveOutputDevices() const = 0;
 		virtual TOptional<FAudioPlatformDeviceInfo> FindDefaultOutputDevice() const = 0;
-		virtual TArray<FAudioPlatformDeviceInfo> GetAggregateDeviceInfo(FName InDeviceID) const = 0;
+
+		/** Determines if the given device Id is that of an aggregate device which uses hardware Ids as device Ids. */
+		virtual bool IsAggregateHardwareDeviceId(const FName InDeviceId) const = 0;
+
+		/** Returns an array of logical audio devices which, when combined, form an aggregate device. */
+		virtual TArray<FAudioPlatformDeviceInfo> GetLogicalAggregateDevices(const FName InHardwareId, const EDeviceEndpointType InEndpointType) const = 0;
 	};
 
 	/** Defines parameters needed for opening a new audio stream to device. */
@@ -245,7 +252,7 @@ namespace Audio
 		/** The number of desired audio frames in audio callback. */
 		uint32 NumFrames;
 		
-		/** The number of queued buffers to use for the strea. */
+		/** The number of queued buffers to use for the stream. */
 		int32 NumBuffers;
 
 		/** Owning platform independent audio mixer ptr.*/
@@ -254,11 +261,17 @@ namespace Audio
 		/** The desired sample rate */
 		uint32 SampleRate;
 
+		/* The maximum number of sources we will try to decode or playback at once. */
+		int32 MaxSources;
+
 		/** Whether or not to try and restore audio to this stream if the audio device is removed (and the device becomes available again). */
 		bool bRestoreIfRemoved;
 
-		/* The maximum number of sources we will try to decode or playback at once. */
-		int32 MaxSources;
+		/** If true, use the system default audio device. If False, AudioDeviceId must contain a valid device Id. */
+		bool bUseSystemAudioDevice;
+
+		/** The device id of an audio output device. May be empty if bUseSystemAudioDevice is true. */
+		FString AudioDeviceId;
 
 		FAudioMixerOpenStreamParams()
 			: OutputDeviceIndex(INDEX_NONE)
@@ -266,8 +279,9 @@ namespace Audio
 			, NumBuffers(1)
 			, AudioMixer(nullptr)
 			, SampleRate(44100)
-			, bRestoreIfRemoved(false)
 			, MaxSources(0)
+			, bRestoreIfRemoved(false)
+			, bUseSystemAudioDevice(true)
 		{}
 	};
 
@@ -607,6 +621,10 @@ namespace Audio
 		/** This is called after InitializeHardware() is called. */
 		AUDIOMIXERCORE_API void PostInitializeHardware();
 
+		/** Used to determine if this object is listening to system device change events. */
+		bool GetIsListeningForDeviceEvents() const { return bIsListeningForDeviceEvents; }
+		void SetIsListeningForDeviceEvents(bool bInListeningForDeviceEvents) { bIsListeningForDeviceEvents = bInListeningForDeviceEvents; }
+
 	protected:
 		
 		// Run the "main" audio device
@@ -770,6 +788,10 @@ namespace Audio
 
 	private:
 		TUniquePtr<FMixerNullCallback> NullDeviceCallback;
+		
+		/** Used to determine if this object is listening to system device change events. */
+		bool bIsListeningForDeviceEvents = true;
+
 	};
 
 	/** Audio mixer platform objects can subclass this in order to add device swap capabilities. */
diff --git a/Engine/Source/Runtime/Engine/Private/AudioDeviceManager.cpp b/Engine/Source/Runtime/Engine/Private/AudioDeviceManager.cpp
--- a/Engine/Source/Runtime/Engine/Private/AudioDeviceManager.cpp
+++ b/Engine/Source/Runtime/Engine/Private/AudioDeviceManager.cpp
@@ -173,6 +173,8 @@ static FAutoConsoleCommand GCleanUpAudioDeviceLeaksCommand(
 FAudioDeviceManager implementation.
 -----------------------------------------------------------------------------*/
 
+bool FAudioDeviceManager::bEnableAggregateDeviceSupport = false;
+
 FAudioDeviceManager::FAudioDeviceManager()
 	: AudioDeviceModule(nullptr)
 	, DeviceIDCounter(0)
@@ -1311,6 +1313,16 @@ void FAudioDeviceManager::SetDynamicSoundVolume(ESoundType SoundType, const FNam
 	DynamicSoundVolumes.FindOrAdd(Key) = FMath::Clamp(Volume, 0.0f, MAX_VOLUME);;
 }
 
+void FAudioDeviceManager::EnableAggregateDeviceSupport(bool bInEnableAggregateDevice)
+{
+	bEnableAggregateDeviceSupport = bInEnableAggregateDevice;
+}
+
+bool FAudioDeviceManager::IsAggregateDeviceSupportEnabled()
+{
+	return bEnableAggregateDeviceSupport;
+}
+
 #if ENABLE_AUDIO_DEBUG
 Audio::FAudioDebugger& FAudioDeviceManager::GetDebugger()
 {
diff --git a/Engine/Source/Runtime/Engine/Public/AudioDeviceManager.h b/Engine/Source/Runtime/Engine/Public/AudioDeviceManager.h
--- a/Engine/Source/Runtime/Engine/Public/AudioDeviceManager.h
+++ b/Engine/Source/Runtime/Engine/Public/AudioDeviceManager.h
@@ -312,6 +312,9 @@ public:
 
 	ENGINE_API Audio::FAudioFormatSettings& GetAudioFormatSettings() const;
 
+	ENGINE_API static void EnableAggregateDeviceSupport(const bool bInEnableAggregateDevice);
+	ENGINE_API static bool IsAggregateDeviceSupportEnabled();
+
 private:
 
 #if ENABLE_AUDIO_DEBUG
@@ -322,6 +325,9 @@ private:
 	TPimplPtr<Audio::FAudioFormatSettings> AudioFormatSettings;
 	TArray<TPimplPtr<FSimpleAudioInfoFactory>> EngineFormats;
 
+	/** True if aggregate devices are supported. Static so it can be set prior to initialization of AudioDeviceManager singleton. */
+	static bool bEnableAggregateDeviceSupport;
+
 	/**
 	 * Conditionally Creates the AudioDeviceManager
 	 * @return New Instance or Existing Singleton, or nullptr on failure (which can happen if audio is disabled).
diff --git a/Engine/Source/Runtime/Windows/AudioMixerWasapi/Private/AudioMixerWasapi.cpp b/Engine/Source/Runtime/Windows/AudioMixerWasapi/Private/AudioMixerWasapi.cpp
--- a/Engine/Source/Runtime/Windows/AudioMixerWasapi/Private/AudioMixerWasapi.cpp
+++ b/Engine/Source/Runtime/Windows/AudioMixerWasapi/Private/AudioMixerWasapi.cpp
@@ -9,10 +9,8 @@
 
 namespace Audio
 {
-	FAudioMixerWasapi::FAudioMixerWasapi(const bool bInUseAggregateDevice) :
-		bUseWasapiAggregateDevice(bInUseAggregateDevice)
+	FAudioMixerWasapi::FAudioMixerWasapi()
 	{
-		CreateDeviceManager(bInUseAggregateDevice, DeviceManager);
 	}
 
 	FAudioMixerWasapi::~FAudioMixerWasapi()
@@ -38,25 +36,7 @@ namespace Audio
 		SCOPED_NAMED_EVENT(FAudioMixerWasapi_InitializeHardware, FColor::Blue);
 
 		RegisterDeviceChangedListener();
-
-		FAudioPlatformSettings EngineSettings = GetPlatformSettings();
-
-		TArray<FWasapiRenderStreamParams> StreamParams;
-		if (!InitStreamParams(AUDIO_MIXER_DEFAULT_DEVICE_INDEX, EngineSettings.CallbackBufferFrameSize, EngineSettings.NumBuffers, EngineSettings.SampleRate, StreamParams))
-		{
-			AUDIO_PLATFORM_LOG_ONCE(TEXT("FAudioMixerWasapi::InitializeHardware InitStreamParams() failed"), Warning);
-			return false;
-		}
-
-		// The ReadNextBufferCallback life cycle is tied to this object. It is ultimately bound to a delegate
-		// in the render stream object which will be unbound in TeardownHardware, prior to 'this' being deallocated.
-		TFunction<void()> ReadNextBufferCallback = [this](){ ReadNextBuffer(); };
-		if (!ensure(DeviceManager && DeviceManager->InitializeHardware(StreamParams, ReadNextBufferCallback)))
-		{
-			AUDIO_PLATFORM_LOG_ONCE(TEXT("FAudioMixerWasapi::InitializeHardware InitializeHardware() failed"), Warning);
-			return false;
-		}
-
+		
 		if (IAudioMixer::ShouldRecycleThreads())
 		{
 			// Pre-create the null render device thread so we can simply wake it up when needed.
@@ -149,11 +129,6 @@ namespace Audio
 		return false;
 	}
 
-	FString FAudioMixerWasapi::GetCurrentDeviceName() const
-	{
-		return AudioStreamInfo.DeviceInfo.Name;
-	}
-
 	bool FAudioMixerWasapi::GetDefaultOutputDeviceIndex(uint32& OutDefaultDeviceIndex) const
 	{
 		OutDefaultDeviceIndex = AUDIO_MIXER_DEFAULT_DEVICE_INDEX;
@@ -177,13 +152,16 @@ namespace Audio
 	bool FAudioMixerWasapi::InitStreamParams(const FAudioPlatformDeviceInfo& InDeviceInfo, const int32 InNumBufferFrames, const int32 InNumBuffers, const int32 InSampleRate, TArray<FWasapiRenderStreamParams>& OutParams) const
 	{
 		SCOPED_NAMED_EVENT(FAudioMixerWasapi_InitStreamParams, FColor::Blue);
+		check(GetDeviceInfoCache());
 
-		if (bUseWasapiAggregateDevice)
+		if (GetDeviceInfoCache()->IsAggregateHardwareDeviceId(*InDeviceInfo.DeviceId))
 		{
 			if (const IAudioPlatformDeviceInfoCache* Cache = GetDeviceInfoCache())
 			{
 				const FName DeviceId = *InDeviceInfo.DeviceId;
-				TArray<FAudioPlatformDeviceInfo> AggregateDevice = Cache->GetAggregateDeviceInfo(DeviceId);
+				// We use the HardwareId as the DeviceId for aggregate devices which is used by
+				// GetAggregateDeviceInfo to gather all the logical devices belonging to this aggregate.
+				TArray<FAudioPlatformDeviceInfo> AggregateDevice = Cache->GetLogicalAggregateDevices(DeviceId, EDeviceEndpointType::Render);
 
 				for (const FAudioPlatformDeviceInfo& AggregateDeviceInfo : AggregateDevice)
 				{
@@ -216,6 +194,7 @@ namespace Audio
 	bool FAudioMixerWasapi::OpenAudioStream(const FAudioMixerOpenStreamParams& Params)
 	{
 		SCOPED_NAMED_EVENT(FAudioMixerWasapi_OpenAudioStream, FColor::Green);
+		check(GetDeviceInfoCache());
 
 		OpenStreamParams = Params;
 
@@ -226,6 +205,10 @@ namespace Audio
 		AudioStreamInfo.NumBuffers = OpenStreamParams.NumBuffers;
 		AudioStreamInfo.AudioMixer = OpenStreamParams.AudioMixer;
 
+		// If the user has selected a specific audio device (not the system default), then
+		// ignore device change events.
+		SetIsListeningForDeviceEvents(Params.bUseSystemAudioDevice);
+
 		TArray<FWasapiRenderStreamParams> StreamParams;
 		if (!InitStreamParams(OpenStreamParams.OutputDeviceIndex, OpenStreamParams.NumFrames, OpenStreamParams.NumBuffers, OpenStreamParams.SampleRate, StreamParams))
 		{
@@ -236,6 +219,29 @@ namespace Audio
 		// be identical because they belong to the same physical device.
 		AudioStreamInfo.DeviceInfo = StreamParams[0].HardwareDeviceInfo;
 
+		// Set the current device name
+		const bool bIsAggregateDevice = GetDeviceInfoCache()->IsAggregateHardwareDeviceId(*Params.AudioDeviceId);
+		if (bIsAggregateDevice)
+		{
+			CurrentDeviceName = ExtractAggregateDeviceName(AudioStreamInfo.DeviceInfo.Name);
+		}
+		else
+		{
+			CurrentDeviceName = AudioStreamInfo.DeviceInfo.Name;
+		}
+
+		// Create and initialize the device manager
+		CreateDeviceManager(bIsAggregateDevice, DeviceManager);
+		
+		// The ReadNextBufferCallback life cycle is tied to this object. It is ultimately bound to a delegate
+		// in the render stream object which will be unbound in TeardownHardware, prior to 'this' being deallocated.
+		TFunction<void()> ReadNextBufferCallback = [this](){ ReadNextBuffer(); };
+		if (!ensure(DeviceManager && DeviceManager->InitializeHardware(StreamParams, ReadNextBufferCallback)))
+		{
+			AUDIO_PLATFORM_LOG_ONCE(TEXT("FAudioMixerWasapi::OpenAudioStream DeviceManager->InitializeHardware() failed"), Warning);
+			return false;
+		}
+
 		// Assign the total number of direct out channels based on the device manager. For WasapiDefaultDeviceMgr this will be 0 
 		// and for WasapiAggregateDeviceMgr this will be the total of all the channels less the main outs (first 8 channels).
 		AudioStreamInfo.DeviceInfo.NumDirectOutChannels = DeviceManager->GetNumDirectOutChannels();
@@ -254,8 +260,8 @@ namespace Audio
 
 		bIsInitialized = true;
 
-		UE_LOG(LogAudioMixer, Display, TEXT("FAudioMixerWasapi initialized: SampleRate=%d NumChannels=%d NumDirectOutChannels=%d AggregateMode=%d"), 
-			OpenStreamParams.SampleRate, AudioStreamInfo.DeviceInfo.NumChannels, AudioStreamInfo.DeviceInfo.NumDirectOutChannels, bUseWasapiAggregateDevice);
+		UE_LOG(LogAudioMixer, Display, TEXT("FAudioMixerWasapi initialized: SampleRate=%d NumChannels=%d NumDirectOutChannels=%d bIsAggregateDevice=%d"), 
+			OpenStreamParams.SampleRate, AudioStreamInfo.DeviceInfo.NumChannels, AudioStreamInfo.DeviceInfo.NumDirectOutChannels, bIsAggregateDevice);
 
 		return true;
 	}
@@ -449,10 +455,10 @@ namespace Audio
 		check(GetDeviceInfoCache());
 
 		// Look up device. Blank name looks up current default.
-		const FName NewDeviceName = *InRequestedDeviceID;
+		const FName NewDeviceId = *InRequestedDeviceID;
 		TOptional<FAudioPlatformDeviceInfo> DeviceInfo;
 	
-		if (TOptional<FAudioPlatformDeviceInfo> TempDeviceInfo = GetDeviceInfoCache()->FindActiveOutputDevice(NewDeviceName))
+		if (TOptional<FAudioPlatformDeviceInfo> TempDeviceInfo = GetDeviceInfoCache()->FindActiveOutputDevice(NewDeviceId))
 		{
 			if (TempDeviceInfo.IsSet())
 			{
@@ -475,6 +481,8 @@ namespace Audio
 
 	bool FAudioMixerWasapi::InitDeviceSwapContextInternal(const FString& InRequestedDeviceID, const TCHAR* InReason, const TOptional<FAudioPlatformDeviceInfo>& InDeviceInfo)
 	{
+		check(GetDeviceInfoCache());
+		
 		// Access to device swap context must be protected by DeviceSwapCriticalSection
 		FScopeLock Lock(&DeviceSwapCriticalSection);
 		
@@ -484,7 +492,7 @@ namespace Audio
 			return false;
 		}
 
-		// Create the device swap context which will valid over the course of the swap
+		// Create the device swap context which will be valid over the course of the swap
 		DeviceSwapContext = MakeUnique<FWasapiDeviceSwapContext>(InRequestedDeviceID, InReason);
 		if (!DeviceSwapContext.IsValid())
 		{
@@ -499,6 +507,8 @@ namespace Audio
 
 		if (DeviceSwapContext->NewDevice.IsSet())
 		{
+			DeviceSwapContext->bIsAggregateDevice = GetDeviceInfoCache()->IsAggregateHardwareDeviceId(*DeviceSwapContext->NewDevice->DeviceId);
+			
 			if (!InitStreamParams(*DeviceSwapContext->NewDevice, EngineSettings.CallbackBufferFrameSize, EngineSettings.NumBuffers, EngineSettings.SampleRate, StreamParams))
 			{
 				UE_LOG(LogAudioMixer, Warning, TEXT("FAudioMixerWasapi::InitializeDeviceSwapContext - InitStreamParams() failed"));
@@ -512,7 +522,6 @@ namespace Audio
 		// happen later in CheckThreadedDeviceSwap from the Game thread.
 		DeviceSwapContext->ReadNextBufferCallback = [this](){ ReadNextBuffer(); };
 		DeviceSwapContext->StreamParams = StreamParams;
-		DeviceSwapContext->bUseAggregateDevice = bUseWasapiAggregateDevice;
 		DeviceSwapContext->PlatformSettings = EngineSettings;
 
 		return true;
@@ -586,14 +595,14 @@ namespace Audio
 			}
 
 			// Don't attempt to create a new setup if there's no devices available.
-			if (!InDeviceContext->NewDevice.IsSet())
+			if (!InDeviceContext->NewDevice.IsSet() || InDeviceContext->StreamParams.IsEmpty())
 			{
-				UE_LOG(LogAudioMixer, Display, TEXT("FAudioMixerWasapi::PerformDeviceSwap - InDeviceContext->bSwitchingToValidDevice: false"));
+				UE_LOG(LogAudioMixer, Display, TEXT("FAudioMixerWasapi::PerformDeviceSwap - no new device to switch to...will run null device"));
 				return {};
 			}
 
 			TUniquePtr<FWasapiDeviceSwapResult> DeviceSwapResult = MakeUnique<FWasapiDeviceSwapResult>();
-			CreateDeviceManager(InDeviceContext->bUseAggregateDevice, DeviceSwapResult->NewDeviceManager);
+			CreateDeviceManager(InDeviceContext->bIsAggregateDevice, DeviceSwapResult->NewDeviceManager);
 			
 			if (!DeviceSwapResult->NewDeviceManager.IsValid())
 			{
@@ -612,9 +621,10 @@ namespace Audio
 				UE_LOG(LogAudioMixer, Warning, TEXT("FAudioMixerWasapi::PerformDeviceSwap - OpenAudioStream failed while attempting to device swap"));
 				return {};
 			}
-
+			
 			DeviceSwapResult->SuccessfulDurationMs = FPlatformTime::ToMilliseconds64(FPlatformTime::Cycles64() - StartTimeCycles);
-			DeviceSwapResult->DeviceInfo = InDeviceContext->NewDevice.Get(FAudioPlatformDeviceInfo());
+			DeviceSwapResult->DeviceInfo = InDeviceContext->StreamParams[0].HardwareDeviceInfo;
+			DeviceSwapResult->bIsAggregateDevice = InDeviceContext->bIsAggregateDevice;
 			DeviceSwapResult->SwapReason = InDeviceContext->DeviceSwapReason;
 			UE_LOG(LogAudioMixer, Display, TEXT("FAudioMixerWasapi::PerformDeviceSwap - successfully completed device swap"));
 
@@ -675,6 +685,16 @@ namespace Audio
 			// Copy our new Device Info into our active one.
 			AudioStreamInfo.DeviceInfo = DeviceSwapResult->DeviceInfo;
 
+			// Set the current device name
+			if (DeviceSwapResult->bIsAggregateDevice)
+			{
+				CurrentDeviceName = ExtractAggregateDeviceName(AudioStreamInfo.DeviceInfo.Name);
+			}
+			else
+			{
+				CurrentDeviceName = AudioStreamInfo.DeviceInfo.Name;
+			}
+
 			// Display our new XAudio2 Mastering voice details.
 			UE_LOG(LogAudioMixer, Display, TEXT("FAudioMixerWasapi::PostDeviceSwap - successful Swap new Device is (NumChannels=%u, SampleRate=%u, DeviceID=%s, Name=%s), Reason=%s, InstanceID=%d, DurationMS=%.2f"),
 				(uint32)AudioStreamInfo.DeviceInfo.NumChannels, (uint32)AudioStreamInfo.DeviceInfo.SampleRate, *AudioStreamInfo.DeviceInfo.DeviceId, *AudioStreamInfo.DeviceInfo.Name,
diff --git a/Engine/Source/Runtime/Windows/AudioMixerWasapi/Private/AudioMixerWasapi.h b/Engine/Source/Runtime/Windows/AudioMixerWasapi/Private/AudioMixerWasapi.h
--- a/Engine/Source/Runtime/Windows/AudioMixerWasapi/Private/AudioMixerWasapi.h
+++ b/Engine/Source/Runtime/Windows/AudioMixerWasapi/Private/AudioMixerWasapi.h
@@ -32,11 +32,11 @@ namespace Audio
 			FDeviceSwapContext(InRequestedDeviceID, InReason)
 		{}
 		
-		bool bUseAggregateDevice = false;
 		FAudioPlatformSettings PlatformSettings;
 		TArray<FWasapiRenderStreamParams> StreamParams;
 		TFunction<void()> ReadNextBufferCallback;
 		TUniquePtr<IAudioMixerWasapiDeviceManager> OldDeviceManager;
+		bool bIsAggregateDevice = false;
 	};
 
 	struct FWasapiDeviceSwapResult : public FDeviceSwapResult
@@ -47,6 +47,7 @@ namespace Audio
 		}
 
 		TUniquePtr<IAudioMixerWasapiDeviceManager> NewDeviceManager;
+		bool bIsAggregateDevice = false;
 	};
 
 	/**
@@ -56,8 +57,7 @@ namespace Audio
 	{
 	public:
 
-		FAudioMixerWasapi() = default;
-		explicit FAudioMixerWasapi(const bool bInUseAggregateDevice);
+		FAudioMixerWasapi();
 		virtual ~FAudioMixerWasapi() override;
 
 		//~ Begin IAudioMixerPlatformInterface
@@ -68,7 +68,6 @@ namespace Audio
 		virtual int32 GetNumFrames(const int32 InNumRequestedFrames) override;
 		virtual bool GetNumOutputDevices(uint32& OutNumOutputDevices) override;
 		virtual bool GetOutputDeviceInfo(const uint32 InDeviceIndex, FAudioPlatformDeviceInfo& OutInfo) override;
-		virtual FString GetCurrentDeviceName() const override;
 		virtual bool GetDefaultOutputDeviceIndex(uint32& OutDefaultDeviceIndex) const override;
 		virtual bool OpenAudioStream(const FAudioMixerOpenStreamParams& Params) override;
 		virtual bool CloseAudioStream() override;
@@ -120,9 +119,6 @@ namespace Audio
 		/** Manages either a single, default device or an aggregate of several devices belonging to the same hardware */
 		TUniquePtr<IAudioMixerWasapiDeviceManager> DeviceManager;
 
-		/** Indicates if this object is aggregating devices or not. */
-		bool bUseWasapiAggregateDevice = false;
-
 		/** Indicates if this object has been successfully initialized. */
 		bool bIsInitialized = false;
 
@@ -132,6 +128,11 @@ namespace Audio
 		/** Fetches an IMMDevice with the given ID. */
 		TComPtr<IMMDevice> GetMMDevice(const FString& InDeviceID) const;
 
+		/** Extracts the hardware device name from a logical device name. The OS
+		 *  places the hardware name in parentheses at the end of the string).
+		 */
+		FString ExtractAggregateDeviceName(const FString& InDeviceID) const;
+		
 		/** Device manager factory */
 		static void CreateDeviceManager(const bool bInUseAggregateDevice, TUniquePtr<IAudioMixerWasapiDeviceManager>& InDeviceManager);
 		
diff --git a/Engine/Source/Runtime/Windows/AudioMixerWasapi/Private/AudioMixerWasapiDeviceListener.cpp b/Engine/Source/Runtime/Windows/AudioMixerWasapi/Private/AudioMixerWasapiDeviceListener.cpp
--- a/Engine/Source/Runtime/Windows/AudioMixerWasapi/Private/AudioMixerWasapiDeviceListener.cpp
+++ b/Engine/Source/Runtime/Windows/AudioMixerWasapi/Private/AudioMixerWasapiDeviceListener.cpp
@@ -7,6 +7,7 @@
 
 #include <atomic>
 
+#include "AudioDeviceManager.h"
 #include "Microsoft/COMPointer.h"
 #include "ScopedCom.h"					// FScopedComString
 
@@ -37,8 +38,11 @@ namespace Audio
 		}
 		if (!DeviceInfoCache.IsValid())
 		{
+			// Wasapi backend supports aggregate devices, provided it is enabled in FAudioDeviceManager
+			bool bIsAggregateDeviceSupported = FAudioDeviceManager::IsAggregateDeviceSupportEnabled();
+
 			// Setup device info cache.
-			DeviceInfoCache = MakeUnique<FWindowsMMDeviceCache>(bUseWasapiAggregateDevice);
+			DeviceInfoCache = MakeUnique<FWindowsMMDeviceCache>(bIsAggregateDeviceSupported);
 			WasapiWinNotificationClient->RegisterDeviceChangedListener(static_cast<FWindowsMMDeviceCache*>(DeviceInfoCache.Get()));
 		}
 
@@ -78,10 +82,15 @@ namespace Audio
 		
 		if (InAudioDeviceRole == EAudioDeviceRole::Console)
 		{
-			UE_LOG(LogAudioMixer, Warning, TEXT("FAudioMixerWasapi: Changing default audio render device to new device: Role=%s, DeviceName=%s, InstanceID=%d"), 
+			UE_LOG(LogAudioMixer, Display, TEXT("FAudioMixerWasapi: Changing default audio render device to new device: Role=%s, DeviceName=%s, InstanceID=%d"), 
 				Audio::ToString(InAudioDeviceRole), *WasapiWinNotificationClient->GetFriendlyName(DeviceId), InstanceID);
 
-			RequestDeviceSwap(DeviceId, /* force */true, TEXT("FAudioMixerWasapi::OnDefaultRenderDeviceChanged"));
+			// Ignore if not listening for events. This means the user has specified a specific device
+			// to use and is not interested in following the current system default device.
+			if (GetIsListeningForDeviceEvents())
+			{
+				RequestDeviceSwap(DeviceId, /* force */true, TEXT("FAudioMixerWasapi::OnDefaultRenderDeviceChanged"));
+			}
 		}
 
 		if (UAudioDeviceNotificationSubsystem* AudioDeviceNotifSubsystem = UAudioDeviceNotificationSubsystem::Get())
@@ -103,10 +112,15 @@ namespace Audio
 		const FString AudioDeviceId = GetOriginalAudioDeviceId();
 		if (AudioStreamInfo.DeviceInfo.DeviceId != AudioDeviceId && DeviceId == AudioDeviceId)
 		{
-			UE_LOG(LogAudioMixer, Warning, TEXT("FAudioMixerWasapi: Original audio device re-added. Moving audio back to original audio device: DeviceName=%s, bRenderDevice=%d, InstanceID=%d"), 
+			UE_LOG(LogAudioMixer, Display, TEXT("FAudioMixerWasapi: Original audio device re-added. Moving audio back to original audio device: DeviceName=%s, bRenderDevice=%d, InstanceID=%d"), 
 				*WasapiWinNotificationClient->GetFriendlyName(*AudioDeviceId), (int32)bIsRenderDevice, InstanceID);
 
-			RequestDeviceSwap(AudioDeviceId, /*force */ true, TEXT("FAudioMixerWasapi::OnDeviceAdded"));
+			// Ignore if not listening for events. This means the user has specified a specific device
+			// to use and is not interested in following the current system default device.
+			if (GetIsListeningForDeviceEvents())
+			{
+				RequestDeviceSwap(AudioDeviceId, /*force */ true, TEXT("FAudioMixerWasapi::OnDeviceAdded"));
+			}
 		}
 
 		if (UAudioDeviceNotificationSubsystem* AudioDeviceNotifSubsystem = UAudioDeviceNotificationSubsystem::Get())
@@ -129,7 +143,12 @@ namespace Audio
 			UE_LOG(LogAudioMixer, Warning, TEXT("FAudioMixerWasapi: Audio device removed [%s], falling back to other windows default device. bIsRenderDevice=%d, InstanceID=%d"), 
 				*WasapiWinNotificationClient->GetFriendlyName(DeviceId), (int32)bIsRenderDevice, InstanceID);
 
-			RequestDeviceSwap(TEXT(""), /* force */ true, TEXT("FAudioMixerWasapi::OnDeviceRemoved"));
+			// Ignore if not listening for events. This means the user has specified a specific device
+			// to use and is not interested in following the current system default device.
+			if (GetIsListeningForDeviceEvents())
+			{
+				RequestDeviceSwap(TEXT(""), /* force */ true, TEXT("FAudioMixerWasapi::OnDeviceRemoved"));
+			}
 		}
 
 		if (UAudioDeviceNotificationSubsystem* AudioDeviceNotifySubsystem = UAudioDeviceNotificationSubsystem::Get())
@@ -166,5 +185,10 @@ namespace Audio
 
 		return TComPtr<IMMDevice>();
 	}
+
+	FString FAudioMixerWasapi::ExtractAggregateDeviceName(const FString& InName) const
+	{
+		return FWindowsMMDeviceCache::ExtractAggregateDeviceName(InName);
+	}
 }
 
diff --git a/Engine/Source/Runtime/Windows/AudioMixerWasapi/Private/AudioMixerWasapiModule.cpp b/Engine/Source/Runtime/Windows/AudioMixerWasapi/Private/AudioMixerWasapiModule.cpp
--- a/Engine/Source/Runtime/Windows/AudioMixerWasapi/Private/AudioMixerWasapiModule.cpp
+++ b/Engine/Source/Runtime/Windows/AudioMixerWasapi/Private/AudioMixerWasapiModule.cpp
@@ -15,6 +15,7 @@ void FAudioMixerModuleWasapi::StartupModule()
 
 	FModuleManager::Get().LoadModuleChecked(TEXT("AudioMixer"));
 	FModuleManager::Get().LoadModuleChecked(TEXT("AudioMixerCore"));
+	FModuleManager::Get().LoadModuleChecked(TEXT("WindowsMMDeviceEnumeration"));
 }
 
 bool FAudioMixerModuleWasapi::IsAudioMixerModule() const
@@ -24,11 +25,7 @@ bool FAudioMixerModuleWasapi::IsAudioMixerModule() const
 
 Audio::IAudioMixerPlatformInterface* FAudioMixerModuleWasapi::CreateAudioMixerPlatformInterface()
 {
-	bool bUseWasapiAggregateDevice = false;
-
-	GConfig->GetBool(TEXT("Audio"), TEXT("UseWasapiAggregateDevice"), bUseWasapiAggregateDevice, GEngineIni);
-
-	return new Audio::FAudioMixerWasapi(bUseWasapiAggregateDevice);
+	return new Audio::FAudioMixerWasapi();
 }
 
 #if PLATFORM_WINDOWS
diff --git a/Engine/Source/Runtime/Windows/AudioMixerXAudio2/Private/AudioMixerPlatformWindows.cpp b/Engine/Source/Runtime/Windows/AudioMixerXAudio2/Private/AudioMixerPlatformWindows.cpp
--- a/Engine/Source/Runtime/Windows/AudioMixerXAudio2/Private/AudioMixerPlatformWindows.cpp
+++ b/Engine/Source/Runtime/Windows/AudioMixerXAudio2/Private/AudioMixerPlatformWindows.cpp
@@ -104,7 +104,12 @@ namespace Audio
 			UE_LOG(LogAudioMixer, Display, TEXT("FMixerPlatformXAudio2: Changing default audio render device to new device: Role=%s, DeviceName=%s, InstanceID=%d"), 
 				Audio::ToString(InAudioDeviceRole), *WindowsNotificationClient->GetFriendlyName(DeviceId), InstanceID);
 
-			RequestDeviceSwap(DeviceId, /* force */true, TEXT("FMixerPlatformXAudio2::OnDefaultRenderDeviceChanged"));
+			// Ignore if not listening for events. This means the user has specified a specific device
+			// to use and is not interested in following the current system default device.
+			if (GetIsListeningForDeviceEvents())
+			{
+				RequestDeviceSwap(DeviceId, /* force */true, TEXT("FMixerPlatformXAudio2::OnDefaultRenderDeviceChanged"));
+			}
 		}
 
 		if (UAudioDeviceNotificationSubsystem* AudioDeviceNotifSubsystem = UAudioDeviceNotificationSubsystem::Get())
@@ -129,7 +134,12 @@ namespace Audio
 			UE_LOG(LogAudioMixer, Display, TEXT("FMixerPlatformXAudio2: Original audio device re-added. Moving audio back to original audio device: DeviceName=%s, bRenderDevice=%d, InstanceID=%d"), 
 				*WindowsNotificationClient->GetFriendlyName(*AudioDeviceId), (int32)bIsRenderDevice, InstanceID);
 
-			RequestDeviceSwap(AudioDeviceId, /*force */ true, TEXT("FMixerPlatformXAudio2::OnDeviceAdded"));
+			// Ignore if not listening for events. This means the user has specified a specific device
+			// to use and is not interested in following the current system default device.
+			if (GetIsListeningForDeviceEvents())
+			{
+				RequestDeviceSwap(AudioDeviceId, /*force */ true, TEXT("FMixerPlatformXAudio2::OnDeviceAdded"));
+			}
 		}
 
 		if (UAudioDeviceNotificationSubsystem* AudioDeviceNotifSubsystem = UAudioDeviceNotificationSubsystem::Get())
@@ -152,7 +162,12 @@ namespace Audio
 			UE_LOG(LogAudioMixer, Display, TEXT("FMixerPlatformXAudio2: Audio device removed [%s], falling back to other windows default device. bIsRenderDevice=%d, InstanceID=%d"), 
 				*WindowsNotificationClient->GetFriendlyName(DeviceId), (int32)bIsRenderDevice, InstanceID);
 
-			RequestDeviceSwap(TEXT(""), /* force */ true, TEXT("FMixerPlatformXAudio2::OnDeviceRemoved"));
+			// Ignore if not listening for events. This means the user has specified a specific device
+			// to use and is not interested in following the current system default device.
+			if (GetIsListeningForDeviceEvents())
+			{
+				RequestDeviceSwap(TEXT(""), /* force */ true, TEXT("FMixerPlatformXAudio2::OnDeviceRemoved"));
+			}
 		}
 
 		if (UAudioDeviceNotificationSubsystem* AudioDeviceNotifSubsystem = UAudioDeviceNotificationSubsystem::Get())
