Fixing build error. Temporarilly removing new deprecation.

#rnx

diff --git a/Engine/Source/Runtime/AudioDeviceEnumeration/Windows/WindowsMMDeviceEnumeration/Public/WindowsMMDeviceInfoCache.h b/Engine/Source/Runtime/AudioDeviceEnumeration/Windows/WindowsMMDeviceEnumeration/Public/WindowsMMDeviceInfoCache.h
--- a/Engine/Source/Runtime/AudioDeviceEnumeration/Windows/WindowsMMDeviceEnumeration/Public/WindowsMMDeviceInfoCache.h
+++ b/Engine/Source/Runtime/AudioDeviceEnumeration/Windows/WindowsMMDeviceEnumeration/Public/WindowsMMDeviceInfoCache.h
@@ -26,7 +26,6 @@ struct WINDOWSMMDEVICEENUMERATION_API FWindowsMMDeviceCache : IAudioMixerDeviceC
 {
 	struct FCacheEntry
 	{
-		UE_DEPRECATED(5.6, "EEndpointType is deprecated. Please use EDeviceEndpointType instead.");
 		enum class EEndpointType { Unknown, Render, Capture };
 
 		FName DeviceId;							// Key
