Build fix for EndpointType deprecation fallout

#rnx

diff --git a/Engine/Source/Runtime/AudioDeviceEnumeration/Windows/WindowsMMDeviceEnumeration/Public/WindowsMMDeviceInfoCache.h b/Engine/Source/Runtime/AudioDeviceEnumeration/Windows/WindowsMMDeviceEnumeration/Public/WindowsMMDeviceInfoCache.h
--- a/Engine/Source/Runtime/AudioDeviceEnumeration/Windows/WindowsMMDeviceEnumeration/Public/WindowsMMDeviceInfoCache.h
+++ b/Engine/Source/Runtime/AudioDeviceEnumeration/Windows/WindowsMMDeviceEnumeration/Public/WindowsMMDeviceInfoCache.h
@@ -26,8 +26,8 @@ struct WINDOWSMMDEVICEENUMERATION_API FWindowsMMDeviceCache : IAudioMixerDeviceC
 {
 	struct FCacheEntry
 	{
-		enum class EEndpointType { Unknown, Render, Capture };
-
+		enum class UE_DEPRECATED(5.6, "EEndpointType is deprecated. Please use EDeviceEndpointType instead.") EEndpointType { Unknown, Render, Capture };
+		
 		FName DeviceId;							// Key
 		FString FriendlyName;
 		FString DeviceFriendlyName;
