Fix crash in audio device enumeration

#rnx

diff --git a/Engine/Source/Runtime/AudioDeviceEnumeration/Windows/WindowsMMDeviceEnumeration/Private/WindowsMMDeviceInfoCache.cpp b/Engine/Source/Runtime/AudioDeviceEnumeration/Windows/WindowsMMDeviceEnumeration/Private/WindowsMMDeviceInfoCache.cpp
--- a/Engine/Source/Runtime/AudioDeviceEnumeration/Windows/WindowsMMDeviceEnumeration/Private/WindowsMMDeviceInfoCache.cpp
+++ b/Engine/Source/Runtime/AudioDeviceEnumeration/Windows/WindowsMMDeviceEnumeration/Private/WindowsMMDeviceInfoCache.cpp
@@ -451,6 +451,12 @@ namespace Audio
 
 	void FWindowsMMDeviceCache::EnumerateDefaults()
 	{
+		if (!DeviceEnumerator.IsValid())
+		{
+			UE_LOG(LogAudioEnumeration, Warning, TEXT("FWindowsMMDeviceCache::EnumerateDefaults() NULL device enumerator"));
+			return;
+		}
+		
 		auto GetDefaultDeviceID = [this](EDataFlow InDataFlow, ERole InRole, FName& OutDeviceId) -> bool
 		{
 			// Mark default device.
@@ -458,11 +464,18 @@ namespace Audio
 			TComPtr<IMMDevice> DefaultDevice;
 			if (SUCCEEDED(DeviceEnumerator->GetDefaultAudioEndpoint(InDataFlow, InRole, &DefaultDevice)))
 			{
-				Audio::FScopeComString DeviceIdString;
-				if (SUCCEEDED(DefaultDevice->GetId(&DeviceIdString.StringPtr)) && DeviceIdString)
+				if (DefaultDevice.IsValid())
+				{
+					Audio::FScopeComString DeviceIdString;
+					if (SUCCEEDED(DefaultDevice->GetId(&DeviceIdString.StringPtr)) && DeviceIdString)
+					{
+						OutDeviceId = DeviceIdString.Get();
+						bSuccess = true;
+					}
+				}
+				else
 				{
-					OutDeviceId = DeviceIdString.Get();
-					bSuccess = true;
+					UE_LOG(LogAudioEnumeration, Display, TEXT("FWindowsMMDeviceCache::EnumerateDefaults() NULL default device for InDataFlow: %d InRole: %d"), InDataFlow, InRole);
 				}
 			}
 			return bSuccess;
