Fix for crash in FWindowsMMDeviceCache::OnDeviceAdded()

[RN] Fix crash when enumerating audio devices.

diff --git a/Engine/Source/Runtime/AudioDeviceEnumeration/Windows/WindowsMMDeviceEnumeration/Private/WindowsMMDeviceInfoCache.cpp b/Engine/Source/Runtime/AudioDeviceEnumeration/Windows/WindowsMMDeviceEnumeration/Private/WindowsMMDeviceInfoCache.cpp
--- a/Engine/Source/Runtime/AudioDeviceEnumeration/Windows/WindowsMMDeviceEnumeration/Private/WindowsMMDeviceInfoCache.cpp
+++ b/Engine/Source/Runtime/AudioDeviceEnumeration/Windows/WindowsMMDeviceEnumeration/Private/WindowsMMDeviceInfoCache.cpp
@@ -519,8 +519,15 @@ namespace Audio
 		if (TOptional<FCacheEntry> NewDeviceEntry = BuildCacheEntry(DeviceId))
 		{
 			FWriteScopeLock WriteLock(CacheMutationLock);
-			check(!Cache.Contains(NewDeviceEntry->DeviceId));
-			Cache.Emplace(NewDeviceEntry->DeviceId, MoveTemp(*NewDeviceEntry));
+			if (FCacheEntry* Entry = Cache.Find(NewDeviceEntry->DeviceId))
+			{
+				*Entry = MoveTemp(*NewDeviceEntry);
+				UE_LOG(LogAudioEnumeration, Warning, TEXT("FWindowsMMDeviceCache::OnDeviceAdded: Device already in cache. Replacing entry for DeviceID='%s'"), *DeviceId);
+			}
+			else
+			{
+				Cache.Emplace(NewDeviceEntry->DeviceId, MoveTemp(*NewDeviceEntry));
+			}
 		}
 		else
 		{
