Minor tweaks to logging

#rnx

diff --git a/Engine/Source/Runtime/AudioDeviceEnumeration/Windows/WindowsMMDeviceEnumeration/Private/WindowsMMDeviceInfoCache.cpp b/Engine/Source/Runtime/AudioDeviceEnumeration/Windows/WindowsMMDeviceEnumeration/Private/WindowsMMDeviceInfoCache.cpp
--- a/Engine/Source/Runtime/AudioDeviceEnumeration/Windows/WindowsMMDeviceEnumeration/Private/WindowsMMDeviceInfoCache.cpp
+++ b/Engine/Source/Runtime/AudioDeviceEnumeration/Windows/WindowsMMDeviceEnumeration/Private/WindowsMMDeviceInfoCache.cpp
@@ -489,12 +489,12 @@ namespace Audio
 			FName DeviceIdName;
 			if (GetDefaultDeviceID(eRender, static_cast<ERole>(i), DeviceIdName))
 			{
-				UE_CLOG(!DeviceIdName.IsNone(), LogAudioEnumeration, Verbose, TEXT("FWindowsMMDeviceCache: Default Render Role='%s', Device='%s'"), ToString((EAudioDeviceRole)i), *GetFriendlyName(DeviceIdName));
+				UE_CLOG(!DeviceIdName.IsNone(), LogAudioEnumeration, Display, TEXT("FWindowsMMDeviceCache: Default Render Role='%s', Device='%s'"), ToString((EAudioDeviceRole)i), *GetFriendlyName(DeviceIdName));
 				DefaultRenderId[i] = DeviceIdName;
 			}
 			if (GetDefaultDeviceID(eCapture, static_cast<ERole>(i), DeviceIdName))
 			{
-				UE_CLOG(!DeviceIdName.IsNone(), LogAudioEnumeration, Verbose, TEXT("FWindowsMMDeviceCache: Default Capture Role='%s', Device='%s'"), ToString((EAudioDeviceRole)i), *GetFriendlyName(DeviceIdName));
+				UE_CLOG(!DeviceIdName.IsNone(), LogAudioEnumeration, Display, TEXT("FWindowsMMDeviceCache: Default Capture Role='%s', Device='%s'"), ToString((EAudioDeviceRole)i), *GetFriendlyName(DeviceIdName));
 				DefaultCaptureId[i] = DeviceIdName;
 			}
 		}
diff --git a/Engine/Source/Runtime/Windows/AudioMixerWasapi/Private/AudioMixerWasapi.cpp b/Engine/Source/Runtime/Windows/AudioMixerWasapi/Private/AudioMixerWasapi.cpp
--- a/Engine/Source/Runtime/Windows/AudioMixerWasapi/Private/AudioMixerWasapi.cpp
+++ b/Engine/Source/Runtime/Windows/AudioMixerWasapi/Private/AudioMixerWasapi.cpp
@@ -169,7 +169,7 @@ namespace Audio
 					TComPtr<IMMDevice> MMDevice = GetMMDevice(AggregateDeviceInfo.DeviceId);
 					if (!MMDevice)
 					{
-						UE_LOG(LogAudioMixer, Error, TEXT("FAudioMixerWasapi::InitStreamParams null MMDevice"));
+						UE_LOG(LogAudioMixer, Error, TEXT("FAudioMixerWasapi::InitStreamParams null MMDevice for aggregate device ID: '%s'"), *InDeviceInfo.DeviceId);
 						return false;
 					}
 
@@ -182,7 +182,7 @@ namespace Audio
 			const TComPtr<IMMDevice> MMDevice = GetMMDevice(InDeviceInfo.DeviceId);
 			if (!MMDevice)
 			{
-				UE_LOG(LogAudioMixer, Error, TEXT("FAudioMixerWasapi::InitStreamParams null MMDevice"));
+				UE_LOG(LogAudioMixer, Error, TEXT("FAudioMixerWasapi::InitStreamParams null MMDevice for device ID: '%s'"), *InDeviceInfo.DeviceId);
 				return false;
 			}
 
