Iwyu target fixes

diff --git a/Engine/Source/Runtime/AVEncoder/Private/MicrosoftCommon.h b/Engine/Source/Runtime/AVEncoder/Private/MicrosoftCommon.h
--- a/Engine/Source/Runtime/AVEncoder/Private/MicrosoftCommon.h
+++ b/Engine/Source/Runtime/AVEncoder/Private/MicrosoftCommon.h
@@ -30,8 +30,6 @@ inline const FString GetComErrorDescription(HRESULT Res)
 	}
 }
 
-#include "Microsoft/AllowMicrosoftPlatformTypes.h"
-
 // macro to deal with COM calls inside a function that returns `false` on error
 #define CHECK_HR(COM_call)\
 	{\
@@ -76,8 +74,6 @@ inline const FString GetComErrorDescription(HRESULT Res)
 		}\
 	}
 
-#include "Microsoft/HideMicrosoftPlatformTypes.h"
-
 // following commented include causes name clash between UE4 and Windows `IMediaEventSink`,
 // we just need a couple of GUIDs from there so the solution is to duplicate them below
 //#include "wmcodecdsp.h"
diff --git a/Engine/Source/Runtime/AudioDeviceEnumeration/Windows/WindowsMMDeviceEnumeration/Private/WindowsMMNotificationClient.cpp b/Engine/Source/Runtime/AudioDeviceEnumeration/Windows/WindowsMMDeviceEnumeration/Private/WindowsMMNotificationClient.cpp
--- a/Engine/Source/Runtime/AudioDeviceEnumeration/Windows/WindowsMMDeviceEnumeration/Private/WindowsMMNotificationClient.cpp
+++ b/Engine/Source/Runtime/AudioDeviceEnumeration/Windows/WindowsMMDeviceEnumeration/Private/WindowsMMNotificationClient.cpp
@@ -127,8 +127,12 @@ namespace Audio
 		Listeners.Add(DeviceChangedListener);
 	}
 
-	#include "Windows/AllowWindowsPlatformAtomics.h"
+}
 
+#include "Windows/AllowWindowsPlatformAtomics.h"
+
+namespace Audio
+{
 	ULONG STDMETHODCALLTYPE FWindowsMMNotificationClient::AddRef()
 	{
 		return InterlockedIncrement(&Ref);
@@ -143,9 +147,12 @@ namespace Audio
 		}
 		return ulRef;
 	}
+}
 
-	#include "Windows/HideWindowsPlatformAtomics.h"
+#include "Windows/HideWindowsPlatformAtomics.h"
 
+namespace Audio
+{
 	HRESULT STDMETHODCALLTYPE FWindowsMMNotificationClient::QueryInterface(const IID& IId, void** UnknownPtrPtr)
 	{
 		// Three rules of QueryInterface: https://docs.microsoft.com/en-us/windows/win32/com/rules-for-implementing-queryinterface
